WorldPartitionHLODUtilities is now an editor plugin
* Enabled by default
* Solve issue where HLOD layer assets could be loaded without their associated builder settings base class being known, resulting in null objects

diff --git a/Engine/Plugins/Editor/WorldPartitionHLODUtilities/Source/Private/HLODBuilderInstancing.cpp b/Engine/Plugins/Editor/WorldPartitionHLODUtilities/Source/Private/HLODBuilderInstancing.cpp
new file mode 100644
--- /dev/null
+++ b/Engine/Plugins/Editor/WorldPartitionHLODUtilities/Source/Private/HLODBuilderInstancing.cpp
@@ -0,0 +1,144 @@
+// Copyright Epic Games, Inc. All Rights Reserved.
+
+#include "HLODBuilderInstancing.h"
+
+#include "Misc/HashBuilder.h"
+#include "Engine/StaticMesh.h"
+#include "Components/InstancedStaticMeshComponent.h"
+#include "WorldPartition/HLOD/HLODActor.h"
+#include "WorldPartition/HLOD/HLODLayer.h"
+
+#include "ISMPartition/ISMComponentDescriptor.h"
+
+namespace 
+{
+	// Instance batcher class based on FISMComponentDescriptor
+	struct FCustomISMComponentDescriptor : public FISMComponentDescriptor
+	{
+		FCustomISMComponentDescriptor(UStaticMeshComponent* SMC)
+		{
+			InitFrom(SMC, false);
+
+			// We'll always want to spawn ISMC, even if our source components are all SMC
+			ComponentClass = UInstancedStaticMeshComponent::StaticClass();
+
+			ComputeHash();
+		}
+	};
+
+	// Store batched instances data
+	struct FInstancingData
+	{
+		int32							NumInstances = 0;
+
+		TArray<FTransform>				InstancesTransforms;
+
+		int32							NumCustomDataFloats = 0;
+		TArray<float>					InstancesCustomData;
+
+		TArray<FInstancedStaticMeshRandomSeed> RandomSeeds;
+	};
+}
+
+UHLODBuilderInstancing::UHLODBuilderInstancing(const FObjectInitializer& ObjectInitializer)
+	: Super(ObjectInitializer)
+{
+}
+
+TArray<UPrimitiveComponent*> UHLODBuilderInstancing::CreateComponents(AWorldPartitionHLOD* InHLODActor, const UHLODLayer* InHLODLayer, const TArray<UPrimitiveComponent*>& InSubComponents) const
+{
+	TRACE_CPUPROFILER_EVENT_SCOPE(UHLODBuilderInstancing::CreateComponents);
+
+	TArray<UPrimitiveComponent*> Components;
+
+	// Prepare instance batches
+	TMap<FISMComponentDescriptor, FInstancingData> InstancesData;
+	for (UPrimitiveComponent* Primitive : InSubComponents)
+	{
+		if (UStaticMeshComponent* SMC = Cast<UStaticMeshComponent>(Primitive))
+		{
+			FCustomISMComponentDescriptor ISMComponentDescriptor(SMC);
+			FInstancingData& InstancingData = InstancesData.FindOrAdd(ISMComponentDescriptor);
+
+			if (UInstancedStaticMeshComponent* ISMC = Cast<UInstancedStaticMeshComponent>(SMC))
+			{
+				InstancingData.NumCustomDataFloats = FMath::Max(InstancingData.NumCustomDataFloats, ISMC->NumCustomDataFloats);
+				InstancingData.RandomSeeds.Add( { InstancingData.NumInstances, ISMC->InstancingRandomSeed } );
+				InstancingData.NumInstances += ISMC->GetInstanceCount();
+			}
+			else
+			{
+				InstancingData.NumInstances++;
+			}
+		}
+	}
+
+	// Resize arrays
+	for (auto& Entry : InstancesData)
+	{
+		const FISMComponentDescriptor& ISMComponentDescriptor = Entry.Key;
+		FInstancingData& EntryInstancingData = Entry.Value;
+
+		EntryInstancingData.InstancesTransforms.Reset(EntryInstancingData.NumInstances);
+		EntryInstancingData.InstancesCustomData.Reset(EntryInstancingData.NumInstances * EntryInstancingData.NumCustomDataFloats);
+	}
+	
+	// Append all transforms & per instance custom data
+	for (UPrimitiveComponent* Primitive : InSubComponents)
+	{
+		if (UStaticMeshComponent* SMC = Cast<UStaticMeshComponent>(Primitive))
+		{
+			FCustomISMComponentDescriptor ISMComponentDescriptor(SMC);
+			FInstancingData& InstancingData = InstancesData.FindChecked(ISMComponentDescriptor);
+
+			if (UInstancedStaticMeshComponent* ISMC = Cast<UInstancedStaticMeshComponent>(SMC))
+			{
+				// Add transforms
+				for (int32 InstanceIdx = 0; InstanceIdx < ISMC->GetInstanceCount(); InstanceIdx++)
+				{
+					FTransform& InstanceTransform = InstancingData.InstancesTransforms.AddDefaulted_GetRef();
+					ISMC->GetInstanceTransform(InstanceIdx, InstanceTransform, true);
+				}
+
+				// Add per instance custom data
+				int32 NumCustomDataFloatToAdd = ISMC->GetInstanceCount() * InstancingData.NumCustomDataFloats;
+				InstancingData.InstancesCustomData.Append(ISMC->PerInstanceSMCustomData);
+				InstancingData.InstancesCustomData.AddDefaulted(NumCustomDataFloatToAdd - ISMC->PerInstanceSMCustomData.Num());
+			}
+			else
+			{
+				InstancingData.InstancesTransforms.Add(SMC->GetComponentTransform());
+				InstancingData.InstancesCustomData.AddDefaulted(InstancingData.NumCustomDataFloats);
+			}
+		}
+	}
+
+	// Create an ISMC for each SM asset we found
+	for (auto& Entry : InstancesData)
+	{
+		const FISMComponentDescriptor& ISMComponentDescriptor = Entry.Key;
+		FInstancingData& EntryInstancingData = Entry.Value;
+
+		check(EntryInstancingData.InstancesTransforms.Num() * EntryInstancingData.NumCustomDataFloats == EntryInstancingData.InstancesCustomData.Num());
+
+		UInstancedStaticMeshComponent* Component = ISMComponentDescriptor.CreateComponent(InHLODActor);
+		Component->SetForcedLodModel(Component->GetStaticMesh()->GetNumLODs());
+		Component->NumCustomDataFloats = EntryInstancingData.NumCustomDataFloats;
+		Component->AddInstances(EntryInstancingData.InstancesTransforms, /*bShouldReturnIndices*/false, /*bWorldSpace*/true);
+		Component->PerInstanceSMCustomData = MoveTemp(EntryInstancingData.InstancesCustomData);
+
+		if (!EntryInstancingData.RandomSeeds.IsEmpty())
+		{
+			Component->InstancingRandomSeed = EntryInstancingData.RandomSeeds[0].RandomSeed;
+		}
+
+		if (EntryInstancingData.RandomSeeds.Num() > 1)
+		{
+			Component->AdditionalRandomSeeds = TArrayView<FInstancedStaticMeshRandomSeed>(&EntryInstancingData.RandomSeeds[1], EntryInstancingData.RandomSeeds.Num() - 1);
+		}
+
+		Components.Add(Component);
+	};
+
+	return Components;
+}
diff --git a/Engine/Plugins/Editor/WorldPartitionHLODUtilities/Source/Private/HLODBuilderMeshApproximate.cpp b/Engine/Plugins/Editor/WorldPartitionHLODUtilities/Source/Private/HLODBuilderMeshApproximate.cpp
new file mode 100644
--- /dev/null
+++ b/Engine/Plugins/Editor/WorldPartitionHLODUtilities/Source/Private/HLODBuilderMeshApproximate.cpp
@@ -0,0 +1,247 @@
+// Copyright Epic Games, Inc. All Rights Reserved.
+
+#include "HLODBuilderMeshApproximate.h"
+
+#include "Algo/ForEach.h"
+#include "Algo/Transform.h"
+#include "WorldPartition/HLOD/HLODActor.h"
+#include "WorldPartition/HLOD/HLODLayer.h"
+
+#include "Engine/StaticMesh.h"
+#include "Materials/Material.h"
+#include "Materials/MaterialInstance.h"
+#include "MaterialUtilities.h"
+
+#include "Modules/ModuleManager.h"
+#include "IGeometryProcessingInterfacesModule.h"
+#include "GeometryProcessingInterfaces/ApproximateActors.h"
+
+#include "Materials/Material.h"
+#include "Engine/HLODProxy.h"
+#include "Serialization/ArchiveCrc32.h"
+
+#include "HLODBuilderInstancing.h"
+
+
+UHLODBuilderMeshApproximate::UHLODBuilderMeshApproximate(const FObjectInitializer& ObjectInitializer)
+	: Super(ObjectInitializer)
+{
+}
+
+UHLODBuilderMeshApproximateSettings::UHLODBuilderMeshApproximateSettings(const FObjectInitializer& ObjectInitializer)
+	: Super(ObjectInitializer)
+{
+#if WITH_EDITORONLY_DATA
+	if (!IsTemplate())
+	{
+		HLODMaterial = GEngine->DefaultHLODFlattenMaterial;
+	}
+#endif
+}
+
+uint32 UHLODBuilderMeshApproximateSettings::GetCRC() const
+{
+	UHLODBuilderMeshApproximateSettings& This = *const_cast<UHLODBuilderMeshApproximateSettings*>(this);
+
+	FArchiveCrc32 Ar;
+
+	Ar << This.MeshApproximationSettings;
+	UE_LOG(LogHLODBuilder, VeryVerbose, TEXT(" - MeshApproximationSettings = %d"), Ar.GetCrc());
+
+	uint32 Hash = Ar.GetCrc();
+
+	if (!HLODMaterial.IsNull())
+	{
+		UMaterialInterface* Material = HLODMaterial.LoadSynchronous();
+		if (Material)
+		{
+			uint32 MaterialCRC = UHLODProxy::GetCRC(Material);
+			UE_LOG(LogHLODBuilder, VeryVerbose, TEXT(" - Material = %d"), MaterialCRC);
+			Hash = HashCombine(Hash, MaterialCRC);
+		}
+	}
+
+	return Hash;
+}
+
+UHLODBuilderSettings* UHLODBuilderMeshApproximate::CreateSettings(UHLODLayer* InHLODLayer) const 
+{
+	UHLODBuilderMeshApproximateSettings* HLODBuilderSettings = NewObject<UHLODBuilderMeshApproximateSettings>(InHLODLayer);
+
+	// If previous settings object is null, this means we have an older version of the object. Populate with the deprecated settings.
+	if (InHLODLayer->GetHLODBuilderSettings() == nullptr)
+	{
+		HLODBuilderSettings->MeshApproximationSettings = InHLODLayer->MeshApproximationSettings_DEPRECATED;
+		HLODBuilderSettings->HLODMaterial = InHLODLayer->HLODMaterial_DEPRECATED;
+	}
+
+	return HLODBuilderSettings;
+}
+
+TArray<UPrimitiveComponent*> UHLODBuilderMeshApproximate::CreateComponents(AWorldPartitionHLOD* InHLODActor, const UHLODLayer* InHLODLayer, const TArray<UPrimitiveComponent*>& InSubComponents) const
+{
+	TRACE_CPUPROFILER_EVENT_SCOPE(UHLODBuilderMeshApproximate::CreateComponents);
+
+	TSet<AActor*> Actors;
+	TArray<UPrimitiveComponent*> InstancedComponents;
+
+	// Filter the input components
+	for (UPrimitiveComponent* SubComponent : InSubComponents)
+	{
+		if (!SubComponent)
+		{
+			continue;
+		}
+
+		switch (SubComponent->HLODBatchingPolicy)
+		{
+		case EHLODBatchingPolicy::None:
+			Actors.Add(SubComponent->GetOwner());
+			break;
+		case EHLODBatchingPolicy::Instancing:
+			InstancedComponents.Add(SubComponent);
+			break;
+		case EHLODBatchingPolicy::MeshSection:
+			InstancedComponents.Add(SubComponent);
+			UE_LOG(LogHLODBuilder, Warning, TEXT("EHLODBatchingPolicy::MeshSection is not yet supported by the MeshApproximate builder."));
+			break;
+		}
+	}
+	
+	IGeometryProcessingInterfacesModule& GeomProcInterfaces = FModuleManager::Get().LoadModuleChecked<IGeometryProcessingInterfacesModule>("GeometryProcessingInterfaces");
+	IGeometryProcessing_ApproximateActors* ApproxActorsAPI = GeomProcInterfaces.GetApproximateActorsImplementation();
+
+	//
+	// Construct options for ApproximateActors operation
+	//
+
+	const UHLODBuilderMeshApproximateSettings* MeshApproximateSettings = CastChecked<UHLODBuilderMeshApproximateSettings>(InHLODLayer->GetHLODBuilderSettings());
+	const FMeshApproximationSettings& UseSettings = MeshApproximateSettings->MeshApproximationSettings;
+	UMaterial* HLODMaterial = MeshApproximateSettings->HLODMaterial.LoadSynchronous();
+
+	IGeometryProcessing_ApproximateActors::FOptions Options = ApproxActorsAPI->ConstructOptions(UseSettings);
+	Options.BasePackagePath = InHLODActor->GetPackage()->GetName();
+	Options.bGenerateLightmapUVs = false;
+	Options.bCreatePhysicsBody = false;
+
+	// Material baking settings
+	Options.BakeMaterial = HLODMaterial;
+	Options.BaseColorTexParamName = FName("BaseColorTexture");
+	Options.NormalTexParamName = FName("NormalTexture");
+	Options.MetallicTexParamName = FName("MetallicTexture");
+	Options.RoughnessTexParamName = FName("RoughnessTexture");
+	Options.SpecularTexParamName = FName("SpecularTexture");
+	Options.EmissiveTexParamName = FName("EmissiveHDRTexture");
+	Options.bUsePackedMRS = true;
+	Options.PackedMRSTexParamName = FName("PackedTexture");
+
+	// Gather bounds of the input components
+	auto GetActorsBounds = [&]() -> FBoxSphereBounds
+	{
+		FBoxSphereBounds Bounds;
+		bool bFirst = true;
+
+		for (UPrimitiveComponent* Component : InSubComponents)
+		{
+			FBoxSphereBounds ComponentBounds = Component->Bounds;
+			Bounds = bFirst ? ComponentBounds : Bounds + ComponentBounds;
+			bFirst = false;
+		}
+
+		return Bounds;
+	};
+	
+	// Compute texel density if needed, depending on the TextureSizingType setting
+	ETextureSizingType TextureSizingType = UseSettings.MaterialSettings.TextureSizingType;
+	float TexelDensityPerMeter = 0.0f;
+
+	IGeometryProcessing_ApproximateActors::ETextureSizePolicy TextureSizePolicy = IGeometryProcessing_ApproximateActors::ETextureSizePolicy::TextureSize;
+	if (TextureSizingType == ETextureSizingType::TextureSizingType_AutomaticFromTexelDensity)
+	{
+		TexelDensityPerMeter = UseSettings.MaterialSettings.TargetTexelDensityPerMeter;
+		TextureSizePolicy = IGeometryProcessing_ApproximateActors::ETextureSizePolicy::TexelDensity;
+	}
+	else if (TextureSizingType == ETextureSizingType::TextureSizingType_AutomaticFromMeshScreenSize)
+	{
+		TexelDensityPerMeter = FMaterialUtilities::ComputeRequiredTexelDensityFromScreenSize(UseSettings.MaterialSettings.MeshMaxScreenSizePercent, GetActorsBounds().SphereRadius);
+		TextureSizePolicy = IGeometryProcessing_ApproximateActors::ETextureSizePolicy::TexelDensity;
+	}
+	else if (TextureSizingType == ETextureSizingType::TextureSizingType_AutomaticFromMeshDrawDistance)
+	{
+		TexelDensityPerMeter = FMaterialUtilities::ComputeRequiredTexelDensityFromDrawDistance(UseSettings.MaterialSettings.MeshMinDrawDistance, GetActorsBounds().SphereRadius);
+		TextureSizePolicy = IGeometryProcessing_ApproximateActors::ETextureSizePolicy::TexelDensity;
+	}
+
+	Options.MeshTexelDensity = TexelDensityPerMeter;
+	Options.TextureSizePolicy = TextureSizePolicy;
+
+	// run actor approximation computation
+	IGeometryProcessing_ApproximateActors::FResults Results;
+	ApproxActorsAPI->ApproximateActors(Actors.Array(), Options, Results);
+
+	TArray<UPrimitiveComponent*> Components;
+	if (Results.ResultCode == IGeometryProcessing_ApproximateActors::EResultCode::Success)
+	{
+		auto FixupAsset = [&InHLODActor](UObject* Asset)
+		{
+			Asset->ClearFlags(RF_Public | RF_Standalone);
+			Asset->Rename(nullptr, InHLODActor->GetPackage(), REN_NonTransactional | REN_DontCreateRedirectors | REN_ForceNoResetLoaders);
+		};
+	
+		Algo::ForEach(Results.NewMeshAssets, FixupAsset);
+		Algo::ForEach(Results.NewMaterials, FixupAsset);
+		Algo::ForEach(Results.NewTextures, FixupAsset);
+
+		for (UStaticMesh* StaticMesh : Results.NewMeshAssets)
+		{
+			UStaticMeshComponent* Component = NewObject<UStaticMeshComponent>(InHLODActor);
+			Component->SetStaticMesh(StaticMesh);
+
+			Components.Add(Component);
+		}
+
+		for (UMaterialInterface* Material : Results.NewMaterials)
+		{
+			UMaterialInstance* MaterialInst = CastChecked<UMaterialInstance>(Material);
+
+			FStaticParameterSet StaticParameterSet;
+			
+			auto SetStaticSwitch = [&StaticParameterSet](FName ParamName, bool bSet)
+			{
+				if (bSet)
+				{
+					FStaticSwitchParameter SwitchParameter;
+					SwitchParameter.ParameterInfo.Name = ParamName;
+					SwitchParameter.Value = true;
+					SwitchParameter.bOverride = true;
+					StaticParameterSet.StaticSwitchParameters.Add(SwitchParameter);
+				}
+			};
+
+			// Set proper switches needed by our base flatten material
+			SetStaticSwitch("UseBaseColor", Options.bBakeBaseColor);
+			SetStaticSwitch("UseRoughness", Options.bBakeRoughness);
+			SetStaticSwitch("UseMetallic", Options.bBakeMetallic);
+			SetStaticSwitch("UseSpecular", Options.bBakeSpecular);
+			SetStaticSwitch("UseEmissiveHDR", Options.bBakeEmissive);
+			SetStaticSwitch("UseNormal", Options.bBakeNormalMap);
+			SetStaticSwitch("PackMetallic", Options.bUsePackedMRS);
+			SetStaticSwitch("PackSpecular", Options.bUsePackedMRS);
+			SetStaticSwitch("PackRoughness", Options.bUsePackedMRS);
+
+			// Force initializing the static permutations according to the switches we have set
+			MaterialInst->UpdateStaticPermutation(StaticParameterSet);
+			MaterialInst->InitStaticPermutation();
+			MaterialInst->PostEditChange();
+		}
+	}
+
+	// Batch instances
+	if (InstancedComponents.Num())
+	{
+		UHLODBuilderInstancing* InstancingHLODBuilder = NewObject<UHLODBuilderInstancing>();
+		Components.Append(InstancingHLODBuilder->CreateComponents(InHLODActor, nullptr, InstancedComponents));
+	}
+
+	return Components;
+}
diff --git a/Engine/Plugins/Editor/WorldPartitionHLODUtilities/Source/Private/HLODBuilderMeshMerge.cpp b/Engine/Plugins/Editor/WorldPartitionHLODUtilities/Source/Private/HLODBuilderMeshMerge.cpp
new file mode 100644
--- /dev/null
+++ b/Engine/Plugins/Editor/WorldPartitionHLODUtilities/Source/Private/HLODBuilderMeshMerge.cpp
@@ -0,0 +1,105 @@
+// Copyright Epic Games, Inc. All Rights Reserved.
+
+#include "HLODBuilderMeshMerge.h"
+
+#include "Algo/ForEach.h"
+#include "Engine/StaticMesh.h"
+#include "Materials/Material.h"
+#include "Components/StaticMeshComponent.h"
+
+#include "WorldPartition/HLOD/HLODActor.h"
+#include "WorldPartition/HLOD/HLODLayer.h"
+
+#include "IMeshMergeUtilities.h"
+#include "MeshMergeModule.h"
+#include "Modules/ModuleManager.h"
+
+#include "Materials/Material.h"
+#include "Engine/HLODProxy.h"
+#include "Serialization/ArchiveCrc32.h"
+
+
+UHLODBuilderMeshMerge::UHLODBuilderMeshMerge(const FObjectInitializer& ObjectInitializer)
+	: Super(ObjectInitializer)
+{
+}
+
+UHLODBuilderMeshMergeSettings::UHLODBuilderMeshMergeSettings(const FObjectInitializer& ObjectInitializer)
+	: Super(ObjectInitializer)
+{
+#if WITH_EDITORONLY_DATA
+	if (!IsTemplate())
+	{
+		HLODMaterial = GEngine->DefaultHLODFlattenMaterial;
+	}
+#endif
+}
+
+uint32 UHLODBuilderMeshMergeSettings::GetCRC() const
+{
+	UHLODBuilderMeshMergeSettings& This = *const_cast<UHLODBuilderMeshMergeSettings*>(this);
+
+	FArchiveCrc32 Ar;
+
+	Ar << This.MeshMergeSettings;
+	UE_LOG(LogHLODBuilder, VeryVerbose, TEXT(" - MeshMergeSettings = %d"), Ar.GetCrc());
+
+	uint32 Hash = Ar.GetCrc();
+
+	if (!HLODMaterial.IsNull())
+	{
+		UMaterialInterface* Material = HLODMaterial.LoadSynchronous();
+		if (Material)
+		{
+			uint32 MaterialCRC = UHLODProxy::GetCRC(Material);
+			UE_LOG(LogHLODBuilder, VeryVerbose, TEXT(" - Material = %d"), MaterialCRC);
+			Hash = HashCombine(Hash, MaterialCRC);
+		}
+	}
+
+	return Hash;
+}
+
+UHLODBuilderSettings* UHLODBuilderMeshMerge::CreateSettings(UHLODLayer* InHLODLayer) const
+{
+	UHLODBuilderMeshMergeSettings* HLODBuilderSettings = NewObject<UHLODBuilderMeshMergeSettings>(InHLODLayer);
+
+	// If previous settings object is null, this means we have an older version of the object. Populate with the deprecated settings.
+	if (InHLODLayer->GetHLODBuilderSettings() == nullptr)
+	{
+		HLODBuilderSettings->MeshMergeSettings = InHLODLayer->MeshMergeSettings_DEPRECATED;
+		HLODBuilderSettings->HLODMaterial = InHLODLayer->HLODMaterial_DEPRECATED;
+	}
+
+	return HLODBuilderSettings;
+}
+
+TArray<UPrimitiveComponent*> UHLODBuilderMeshMerge::CreateComponents(AWorldPartitionHLOD* InHLODActor, const UHLODLayer* InHLODLayer, const TArray<UPrimitiveComponent*>& InSubComponents) const
+{
+	TRACE_CPUPROFILER_EVENT_SCOPE(UHLODBuilderMeshMerge::CreateComponents);
+
+	TArray<UObject*> Assets;
+	FVector MergedActorLocation;
+
+	const UHLODBuilderMeshMergeSettings* MeshMergeSettings = CastChecked<UHLODBuilderMeshMergeSettings>(InHLODLayer->GetHLODBuilderSettings());
+	const FMeshMergingSettings& UseSettings = MeshMergeSettings->MeshMergeSettings;
+	UMaterial* HLODMaterial = MeshMergeSettings->HLODMaterial.LoadSynchronous();
+
+	const IMeshMergeUtilities& MeshMergeUtilities = FModuleManager::Get().LoadModuleChecked<IMeshMergeModule>("MeshMergeUtilities").GetUtilities();
+	MeshMergeUtilities.MergeComponentsToStaticMesh(InSubComponents, InHLODActor->GetWorld(), UseSettings, HLODMaterial, InHLODActor->GetPackage(), InHLODActor->GetActorLabel(), Assets, MergedActorLocation, 0.25f, false);
+
+	UStaticMeshComponent* Component = nullptr;
+	Algo::ForEach(Assets, [this, InHLODActor, &Component, &MergedActorLocation](UObject* Asset)
+	{
+		Asset->ClearFlags(RF_Public | RF_Standalone);
+
+		if (Cast<UStaticMesh>(Asset))
+		{
+			Component = NewObject<UStaticMeshComponent>(InHLODActor);
+			Component->SetStaticMesh(static_cast<UStaticMesh*>(Asset));
+			Component->SetWorldLocation(MergedActorLocation);
+		}
+	});
+
+	return TArray<UPrimitiveComponent*>({ Component });
+}
diff --git a/Engine/Plugins/Editor/WorldPartitionHLODUtilities/Source/Private/HLODBuilderMeshSimplify.cpp b/Engine/Plugins/Editor/WorldPartitionHLODUtilities/Source/Private/HLODBuilderMeshSimplify.cpp
new file mode 100644
--- /dev/null
+++ b/Engine/Plugins/Editor/WorldPartitionHLODUtilities/Source/Private/HLODBuilderMeshSimplify.cpp
@@ -0,0 +1,147 @@
+// Copyright Epic Games, Inc. All Rights Reserved.
+
+#include "HLODBuilderMeshSimplify.h"
+
+#include "Algo/ForEach.h"
+#include "Algo/Transform.h"
+#include "Engine/StaticMesh.h"
+#include "Materials/Material.h"
+#include "Components/StaticMeshComponent.h"
+
+#include "WorldPartition/HLOD/HLODActor.h"
+#include "WorldPartition/HLOD/HLODLayer.h"
+
+#include "IMeshMergeUtilities.h"
+#include "MeshMergeModule.h"
+#include "Modules/ModuleManager.h"
+
+#include "Materials/Material.h"
+#include "Engine/HLODProxy.h"
+#include "Serialization/ArchiveCrc32.h"
+
+#include "HLODBuilderInstancing.h"
+
+
+UHLODBuilderMeshSimplify::UHLODBuilderMeshSimplify(const FObjectInitializer& ObjectInitializer)
+	: Super(ObjectInitializer)
+{
+}
+
+UHLODBuilderMeshSimplifySettings::UHLODBuilderMeshSimplifySettings(const FObjectInitializer& ObjectInitializer)
+	: Super(ObjectInitializer)
+{
+#if WITH_EDITORONLY_DATA
+	if (!IsTemplate())
+	{
+		HLODMaterial = GEngine->DefaultHLODFlattenMaterial;
+	}
+#endif
+}
+
+uint32 UHLODBuilderMeshSimplifySettings::GetCRC() const
+{
+	UHLODBuilderMeshSimplifySettings& This = *const_cast<UHLODBuilderMeshSimplifySettings*>(this);
+
+	FArchiveCrc32 Ar;
+
+	Ar << This.MeshSimplifySettings;
+	UE_LOG(LogHLODBuilder, VeryVerbose, TEXT(" - MeshSimplifySettings = %d"), Ar.GetCrc());
+
+	uint32 Hash = Ar.GetCrc();
+
+	if (!HLODMaterial.IsNull())
+	{
+		UMaterialInterface* Material = HLODMaterial.LoadSynchronous();
+		if (Material)
+		{
+			uint32 MaterialCRC = UHLODProxy::GetCRC(Material);
+			UE_LOG(LogHLODBuilder, VeryVerbose, TEXT(" - Material = %d"), MaterialCRC);
+			Hash = HashCombine(Hash, MaterialCRC);
+		}
+	}
+
+	return Hash;
+}
+
+UHLODBuilderSettings* UHLODBuilderMeshSimplify::CreateSettings(UHLODLayer* InHLODLayer) const
+{
+	UHLODBuilderMeshSimplifySettings* HLODBuilderSettings = NewObject<UHLODBuilderMeshSimplifySettings>(InHLODLayer);
+
+	// If previous settings object is null, this means we have an older version of the object. Populate with the deprecated settings.
+	if (InHLODLayer->GetHLODBuilderSettings() == nullptr)
+	{
+		HLODBuilderSettings->MeshSimplifySettings = InHLODLayer->MeshSimplifySettings_DEPRECATED;
+		HLODBuilderSettings->HLODMaterial = InHLODLayer->HLODMaterial_DEPRECATED;
+	}
+
+	return HLODBuilderSettings;
+}
+
+TArray<UPrimitiveComponent*> UHLODBuilderMeshSimplify::CreateComponents(AWorldPartitionHLOD* InHLODActor, const UHLODLayer* InHLODLayer, const TArray<UPrimitiveComponent*>& InSubComponents) const
+{
+	TRACE_CPUPROFILER_EVENT_SCOPE(UHLODBuilderMeshSimplifySettings::CreateComponents);
+
+	TArray<UObject*> Assets;
+	FCreateProxyDelegate ProxyDelegate;
+	ProxyDelegate.BindLambda([&Assets](const FGuid Guid, TArray<UObject*>& InAssetsCreated) { Assets = InAssetsCreated; });
+
+	TArray<UStaticMeshComponent*> StaticMeshComponents;
+	TArray<UPrimitiveComponent*> InstancedComponents;
+
+	// Filter the input components
+	for (UPrimitiveComponent* SubComponent : InSubComponents)
+	{
+		if (!SubComponent)
+		{
+			continue;
+		}
+
+		switch (SubComponent->HLODBatchingPolicy)
+		{
+		case EHLODBatchingPolicy::None:
+			if (UStaticMeshComponent* SMC = Cast<UStaticMeshComponent>(SubComponent))
+			{
+				StaticMeshComponents.Add(SMC);
+			}
+			break;
+		case EHLODBatchingPolicy::Instancing:
+			InstancedComponents.Add(SubComponent);
+			break;
+		case EHLODBatchingPolicy::MeshSection:
+			InstancedComponents.Add(SubComponent);
+			UE_LOG(LogHLODBuilder, Warning, TEXT("EHLODBatchingPolicy::MeshSection is not yet supported by the MeshApproximate builder."));
+			break;
+		}
+	}
+
+	const UHLODBuilderMeshSimplifySettings* MeshSimplifySettings = CastChecked<UHLODBuilderMeshSimplifySettings>(InHLODLayer->GetHLODBuilderSettings());
+	const FMeshProxySettings& UseSettings = MeshSimplifySettings->MeshSimplifySettings;
+	UMaterial* HLODMaterial = MeshSimplifySettings->HLODMaterial.LoadSynchronous();
+
+	const IMeshMergeUtilities& MeshMergeUtilities = FModuleManager::Get().LoadModuleChecked<IMeshMergeModule>("MeshMergeUtilities").GetUtilities();
+	MeshMergeUtilities.CreateProxyMesh(StaticMeshComponents, UseSettings, HLODMaterial, InHLODActor->GetPackage(), InHLODActor->GetActorLabel(), FGuid::NewGuid(), ProxyDelegate, true);
+
+	UStaticMeshComponent* Component = nullptr;
+	Algo::ForEach(Assets, [this, InHLODActor, &Component](UObject* Asset)
+	{
+		Asset->ClearFlags(RF_Public | RF_Standalone);
+
+		if (Cast<UStaticMesh>(Asset))
+		{
+			Component = NewObject<UStaticMeshComponent>(InHLODActor);
+			Component->SetStaticMesh(static_cast<UStaticMesh*>(Asset));
+		}
+	});
+
+	TArray<UPrimitiveComponent*> Components;
+	Components.Add(Component);
+
+	// Batch instances
+	if (InstancedComponents.Num())
+	{
+		UHLODBuilderInstancing* InstancingHLODBuilder = NewObject<UHLODBuilderInstancing>();
+		Components.Append(InstancingHLODBuilder->CreateComponents(InHLODActor, nullptr, InstancedComponents));
+	}
+
+	return Components;
+}
diff --git a/Engine/Plugins/Editor/WorldPartitionHLODUtilities/Source/Private/WorldPartitionHLODUtilities.cpp b/Engine/Plugins/Editor/WorldPartitionHLODUtilities/Source/Private/WorldPartitionHLODUtilities.cpp
new file mode 100644
--- /dev/null
+++ b/Engine/Plugins/Editor/WorldPartitionHLODUtilities/Source/Private/WorldPartitionHLODUtilities.cpp
@@ -0,0 +1,411 @@
+// Copyright Epic Games, Inc. All Rights Reserved.
+
+#include "WorldPartitionHLODUtilities.h"
+
+#if WITH_EDITOR
+
+#include "WorldPartition/WorldPartition.h"
+#include "WorldPartition/WorldPartitionHandle.h"
+#include "WorldPartition/HLOD/HLODActor.h"
+#include "WorldPartition/HLOD/HLODActorDesc.h"
+#include "WorldPartition/HLOD/HLODBuilder.h"
+#include "WorldPartition/HLOD/HLODLayer.h"
+
+#include "ISMPartition/ISMComponentDescriptor.h"
+
+#include "Serialization/ArchiveCrc32.h"
+#include "Templates/UniquePtr.h"
+
+#include "Engine/StaticMesh.h"
+#include "Engine/HLODProxy.h"
+#include "Serialization/ArchiveCrc32.h"
+#include "Materials/Material.h"
+#include "AssetCompilingManager.h"
+
+#include "LevelInstance/LevelInstanceActor.h"
+#include "LevelInstance/LevelInstanceSubsystem.h"
+
+#include "HLODBuilderInstancing.h"
+#include "HLODBuilderMeshMerge.h"
+#include "HLODBuilderMeshSimplify.h"
+#include "HLODBuilderMeshApproximate.h"
+
+static bool LoadSubActors(AWorldPartitionHLOD* InHLODActor, TArray<FWorldPartitionReference>& OutActors)
+{
+	OutActors.Reserve(InHLODActor->GetSubActors().Num());
+
+	UWorld* World = InHLODActor->GetWorld();
+	UWorldPartition* WorldPartition = World->GetWorldPartition();
+	check(WorldPartition);
+
+	bool bIsDirty = false;
+
+	// Gather (and potentially load) actors
+	for (const FGuid& SubActorGuid : InHLODActor->GetSubActors())
+	{
+		FWorldPartitionReference ActorRef(WorldPartition, SubActorGuid);
+
+		if (ActorRef.IsValid())
+		{
+			AActor* LoadedActor = ActorRef.Get()->GetActor();
+
+			// Load level instances
+			if (ALevelInstance* LevelInstance = Cast<ALevelInstance>(LoadedActor))
+			{
+				// Wait for level instance loading
+				if (LevelInstance->SupportsLoading())
+				{
+					LevelInstance->GetLevelInstanceSubsystem()->BlockLoadLevelInstance(LevelInstance);
+				}
+			}
+
+			OutActors.Add(MoveTemp(ActorRef));
+		}
+		else
+		{
+			bIsDirty = true;
+		}
+	}
+
+	return !bIsDirty;
+}
+
+static uint32 GetCRC(const UHLODLayer* InHLODLayer)
+{
+	UHLODLayer& HLODLayer= *const_cast<UHLODLayer*>(InHLODLayer);
+
+	uint32 CRC;
+
+	CRC = GetTypeHash(HLODLayer.GetLayerType());
+	UE_LOG(LogHLODBuilder, VeryVerbose, TEXT(" - LayerType = %d"), CRC);
+
+	CRC = HashCombine(HLODLayer.GetHLODBuilderSettings()->GetCRC(), CRC);
+	UE_LOG(LogHLODBuilder, VeryVerbose, TEXT(" - HLODBuilderSettings = %d"), CRC);
+
+	CRC = HashCombine(HLODLayer.GetCellSize(), CRC);
+	UE_LOG(LogHLODBuilder, VeryVerbose, TEXT(" - CellSize = %d"), CRC);
+
+	return CRC;
+}
+
+static uint32 ComputeHLODHash(AWorldPartitionHLOD* InHLODActor, const TArray<FWorldPartitionReference>& InActors)
+{
+	FArchiveCrc32 Ar;
+
+	// Base key, changing this will force a rebuild of all HLODs
+	FString HLODBaseKey = "4CE9431F9C6842F996786C980641B63A";
+	Ar << HLODBaseKey;
+
+	// HLOD Layer
+	uint32 HLODLayerHash = GetCRC(InHLODActor->GetSubActorsHLODLayer());
+	UE_LOG(LogHLODBuilder, VeryVerbose, TEXT(" - HLODLayer (%s) = %x"), *InHLODActor->GetSubActorsHLODLayer()->GetName(), HLODLayerHash);
+	Ar << HLODLayerHash;
+
+	// We get the CRC of each component
+	TArray<uint32> ComponentsCRCs;
+	for (UPrimitiveComponent* Component : UHLODBuilder::GatherPrimitiveComponents(InActors))
+	{
+		if (UStaticMeshComponent* StaticMeshComponent = Cast<UStaticMeshComponent>(Component))
+		{
+			uint32 ComponentCRC = 0;
+
+			UE_LOG(LogHLODBuilder, VeryVerbose, TEXT(" - Component \'%s\' from actor \'%s\'"), *Component->GetName(), *Component->GetOwner()->GetName());
+
+			// CRC component
+			uint32 StaticMeshComponentCRC = UHLODProxy::GetCRC(StaticMeshComponent);
+			UE_LOG(LogHLODBuilder, VeryVerbose, TEXT("     - StaticMeshComponent (%s) = %x"), *StaticMeshComponent->GetName(), StaticMeshComponentCRC);
+			ComponentCRC = HashCombine(ComponentCRC, StaticMeshComponentCRC);
+
+			if (UStaticMesh* StaticMesh = StaticMeshComponent->GetStaticMesh())
+			{
+				// CRC static mesh
+				int32 StaticMeshCRC = UHLODProxy::GetCRC(StaticMesh);
+				UE_LOG(LogHLODBuilder, VeryVerbose, TEXT("     - StaticMesh (%s) = %x"), *StaticMesh->GetName(), StaticMeshCRC);
+				ComponentCRC = HashCombine(ComponentCRC, StaticMeshCRC);
+
+				// CRC materials
+				const int32 NumMaterials = StaticMeshComponent->GetNumMaterials();
+				for (int32 MaterialIndex = 0; MaterialIndex < NumMaterials; ++MaterialIndex)
+				{
+					UMaterialInterface* MaterialInterface = StaticMeshComponent->GetMaterial(MaterialIndex);
+					if (MaterialInterface)
+					{
+						uint32 MaterialInterfaceCRC = UHLODProxy::GetCRC(MaterialInterface);
+						UE_LOG(LogHLODBuilder, VeryVerbose, TEXT("     - MaterialInterface (%s) = %x"), *MaterialInterface->GetName(), MaterialInterfaceCRC);
+						ComponentCRC = HashCombine(ComponentCRC, MaterialInterfaceCRC);
+
+						TArray<UTexture*> Textures;
+						MaterialInterface->GetUsedTextures(Textures, EMaterialQualityLevel::High, true, ERHIFeatureLevel::SM5, true);
+						for (UTexture* Texture : Textures)
+						{
+							uint32 TextureCRC = UHLODProxy::GetCRC(Texture);
+							UE_LOG(LogHLODBuilder, VeryVerbose, TEXT("     - Texture (%s) = %x"), *Texture->GetName(), TextureCRC);
+							ComponentCRC = HashCombine(ComponentCRC, TextureCRC);
+						}
+					}
+				}
+			}
+
+			ComponentsCRCs.Add(ComponentCRC);
+		}
+	}
+
+	// Sort the components CRCs to ensure the order of components won't have an impact on the final CRC
+	ComponentsCRCs.Sort();
+
+	// Append all components CRCs
+	Ar << ComponentsCRCs;
+
+	return Ar.GetCrc();
+}
+
+TArray<AWorldPartitionHLOD*> FWorldPartitionHLODUtilities::CreateHLODActors(FHLODCreationContext& InCreationContext, const FHLODCreationParams& InCreationParams, const TSet<FActorInstance>& InActors, const TArray<const UDataLayer*>& InDataLayers)
+{
+	struct FSubActorsInfo
+	{
+		TArray<FGuid>			SubActors;
+		bool					bIsSpatiallyLoaded;
+	};
+	TMap<UHLODLayer*, FSubActorsInfo> SubActorsInfos;
+
+	for (const FActorInstance& ActorInstance : InActors)
+	{
+		const FWorldPartitionActorDescView& ActorDescView = ActorInstance.GetActorDescView();
+		if (ActorDescView.GetActorIsHLODRelevant())
+		{
+			UHLODLayer* HLODLayer = UHLODLayer::GetHLODLayer(ActorDescView, InCreationParams.WorldPartition);
+			if (HLODLayer)
+			{
+				FSubActorsInfo& SubActorsInfo = SubActorsInfos.FindOrAdd(HLODLayer);
+
+				SubActorsInfo.SubActors.Add(ActorInstance.Actor);
+				if (ActorDescView.GetIsSpatiallyLoaded())
+				{
+					SubActorsInfo.bIsSpatiallyLoaded = true;
+				}
+			}
+		}
+	}
+
+	TArray<AWorldPartitionHLOD*> HLODActors;
+	for (const auto& Pair : SubActorsInfos)
+	{
+		const UHLODLayer* HLODLayer = Pair.Key;
+		const FSubActorsInfo& SubActorsInfo = Pair.Value;
+		check(!SubActorsInfo.SubActors.IsEmpty());
+
+		// Compute HLODActor hash
+		uint64 CellHash = FHLODActorDesc::ComputeCellHash(HLODLayer->GetName(), InCreationParams.GridIndexX, InCreationParams.GridIndexY, InCreationParams.GridIndexZ, InCreationParams.DataLayersID);
+
+		AWorldPartitionHLOD* HLODActor = nullptr;
+		FWorldPartitionHandle HLODActorHandle;
+		if (InCreationContext.HLODActorDescs.RemoveAndCopyValue(CellHash, HLODActorHandle))
+		{
+			InCreationContext.ActorReferences.Add(HLODActorHandle);
+			HLODActor = CastChecked<AWorldPartitionHLOD>(HLODActorHandle->GetActor());
+		}
+
+		bool bNewActor = HLODActor == nullptr;
+		if (bNewActor)
+		{
+			FActorSpawnParameters SpawnParams;
+			SpawnParams.Name = *FString::Printf(TEXT("%s_%016llx"), *HLODLayer->GetName(), CellHash);
+			SpawnParams.NameMode = FActorSpawnParameters::ESpawnActorNameMode::Required_Fatal;
+			HLODActor = InCreationParams.WorldPartition->GetWorld()->SpawnActor<AWorldPartitionHLOD>(SpawnParams);
+
+			HLODActor->SetActorLabel(FString::Printf(TEXT("HLOD%d_%s"), InCreationParams.HLODLevel, *InCreationParams.CellName.ToString()));
+			HLODActor->SetFolderPath(*FString::Printf(TEXT("HLOD/HLOD%d"), InCreationParams.HLODLevel));
+			HLODActor->SetSourceCellName(InCreationParams.CellName);
+			HLODActor->SetSubActorsHLODLayer(HLODLayer);
+			HLODActor->SetGridIndices(InCreationParams.GridIndexX, InCreationParams.GridIndexY, InCreationParams.GridIndexZ);
+
+			// Make sure the generated HLOD actor has the same data layers as the source actors
+			for (const UDataLayer* DataLayer : InDataLayers)
+			{
+				HLODActor->AddDataLayer(DataLayer);
+			}
+		}
+		else
+		{
+#if DO_CHECK
+			uint64 GridIndexX, GridIndexY, GridIndexZ;
+			HLODActor->GetGridIndices(GridIndexX, GridIndexY, GridIndexZ);
+			check(GridIndexX == InCreationParams.GridIndexX);
+			check(GridIndexY == InCreationParams.GridIndexY);
+			check(GridIndexZ == InCreationParams.GridIndexZ);
+			check(HLODActor->GetSubActorsHLODLayer() == HLODLayer);
+			check(FDataLayersID(HLODActor->GetDataLayerObjects()) == InCreationParams.DataLayersID);
+#endif
+		}
+
+		bool bIsDirty = false;
+
+		// Sub actors
+		{
+			bool bSubActorsChanged = HLODActor->GetSubActors().Num() != SubActorsInfo.SubActors.Num();
+			if (!bSubActorsChanged)
+			{
+				TArray<FGuid> A = HLODActor->GetSubActors();
+				TArray<FGuid> B = SubActorsInfo.SubActors;
+				A.Sort();
+				B.Sort();
+				bSubActorsChanged = A != B;
+			}
+
+			if (bSubActorsChanged)
+			{
+				HLODActor->SetSubActors(SubActorsInfo.SubActors);
+				bIsDirty = true;
+			}
+		}
+
+		// Runtime grid
+		FName RuntimeGrid = HLODLayer->GetRuntimeGrid(InCreationParams.HLODLevel);
+		if (HLODActor->GetRuntimeGrid() != RuntimeGrid)
+		{
+			HLODActor->SetRuntimeGrid(RuntimeGrid);
+			bIsDirty = true;
+		}
+
+		// Spatially loaded
+		// HLOD that are always loaded will not take the SubActorsInfo.GridPlacement into account
+		bool bExpectedIsSpatiallyLoaded = !HLODLayer->IsSpatiallyLoaded() ? false : SubActorsInfo.bIsSpatiallyLoaded;
+		if (HLODActor->GetIsSpatiallyLoaded() != bExpectedIsSpatiallyLoaded)
+		{
+			HLODActor->SetIsSpatiallyLoaded(bExpectedIsSpatiallyLoaded);
+			bIsDirty = true;
+		}
+
+		// HLOD level
+		if (HLODActor->GetLODLevel() != InCreationParams.HLODLevel)
+		{
+			HLODActor->SetLODLevel(InCreationParams.HLODLevel);
+			bIsDirty = true;
+		}
+
+		// Require warmup
+		if (HLODActor->DoesRequireWarmup() != HLODLayer->DoesRequireWarmup())
+		{
+			HLODActor->SetRequireWarmup(HLODLayer->DoesRequireWarmup());
+			bIsDirty = true;
+		}
+
+		// Parent HLOD layer
+		UHLODLayer* ParentHLODLayer = HLODLayer->GetParentLayer().LoadSynchronous();
+		if (HLODActor->GetHLODLayer() != ParentHLODLayer)
+		{
+			HLODActor->SetHLODLayer(ParentHLODLayer);
+			bIsDirty = true;
+		}
+
+		// Cell bounds
+		if (!HLODActor->GetHLODBounds().Equals(InCreationParams.CellBounds))
+		{
+			HLODActor->SetHLODBounds(InCreationParams.CellBounds);
+			bIsDirty = true;
+		}
+
+		// If any change was performed, mark HLOD package as dirty
+		if (bIsDirty)
+		{
+			HLODActor->MarkPackageDirty();
+		}
+
+		HLODActors.Add(HLODActor);
+	}
+
+	return HLODActors;
+}
+
+TSubclassOf<UHLODBuilder> FWorldPartitionHLODUtilities::GetHLODBuilderClass(const UHLODLayer* InHLODLayer)
+{
+	EHLODLayerType HLODLayerType = InHLODLayer->GetLayerType();
+	switch (HLODLayerType)
+	{
+	case EHLODLayerType::Instancing:
+		return UHLODBuilderInstancing::StaticClass();
+		break;
+
+	case EHLODLayerType::MeshMerge:
+		return UHLODBuilderMeshMerge::StaticClass();
+		break;
+
+	case EHLODLayerType::MeshSimplify:
+		return UHLODBuilderMeshSimplify::StaticClass();
+		break;
+
+	case EHLODLayerType::MeshApproximate:
+		return UHLODBuilderMeshApproximate::StaticClass();
+		break;
+
+	case EHLODLayerType::Custom:
+		return InHLODLayer->GetHLODBuilderClass();
+		break;
+
+	default:
+		checkf(false, TEXT("Unsupported type"));
+		return nullptr;
+	}
+}
+
+UHLODBuilderSettings* FWorldPartitionHLODUtilities::CreateHLODBuilderSettings(UHLODLayer* InHLODLayer)
+{
+	TSubclassOf<UHLODBuilder> HLODBuilderClass = GetHLODBuilderClass(InHLODLayer);
+	if (!HLODBuilderClass)
+	{
+		return NewObject<UHLODBuilderSettings>(InHLODLayer, UHLODBuilderSettings::StaticClass());
+	}
+
+	UHLODBuilderSettings* HLODBuilderSettings = HLODBuilderClass->GetDefaultObject<UHLODBuilder>()->CreateSettings(InHLODLayer);
+	if (!ensure(HLODBuilderSettings))
+	{
+		return NewObject<UHLODBuilderSettings>(InHLODLayer, UHLODBuilderSettings::StaticClass());
+	}
+
+	return HLODBuilderSettings;
+}
+
+uint32 FWorldPartitionHLODUtilities::BuildHLOD(AWorldPartitionHLOD* InHLODActor)
+{
+	TArray<FWorldPartitionReference> SubActors;
+	
+	bool bIsDirty = !LoadSubActors(InHLODActor, SubActors);
+	if (bIsDirty)
+	{
+		UE_LOG(LogHLODBuilder, Warning, TEXT("HLOD actor \"%s\" needs to be rebuilt as it references actors that have been deleted."), *InHLODActor->GetActorLabel());
+	}
+
+	uint32 OldHLODHash = bIsDirty ? 0 : InHLODActor->GetHLODHash();
+	uint32 NewHLODHash = ComputeHLODHash(InHLODActor, SubActors);
+
+	if (OldHLODHash == NewHLODHash)
+	{
+		return OldHLODHash;
+	}
+
+	const UHLODLayer* HLODLayer = InHLODActor->GetSubActorsHLODLayer();
+	TSubclassOf<UHLODBuilder> HLODBuilderClass = GetHLODBuilderClass(HLODLayer);
+
+	if (HLODBuilderClass)
+	{
+		UHLODBuilder* HLODBuilder = NewObject<UHLODBuilder>(GetTransientPackage(), HLODBuilderClass);
+		if (ensure(HLODBuilder))
+		{
+			HLODBuilder->AddToRoot();
+			if (HLODBuilder->RequiresCompiledAssets())
+			{
+				// Wait for compilation to finish
+				FAssetCompilingManager::Get().FinishAllCompilation();
+			}
+
+			HLODBuilder->Build(InHLODActor, HLODLayer, SubActors);
+			HLODBuilder->RemoveFromRoot();
+		}
+			
+		InHLODActor->MarkPackageDirty();
+	}
+
+	return NewHLODHash;
+}
+
+#endif
diff --git a/Engine/Plugins/Editor/WorldPartitionHLODUtilities/Source/Private/WorldPartitionHLODUtilitiesModule.cpp b/Engine/Plugins/Editor/WorldPartitionHLODUtilities/Source/Private/WorldPartitionHLODUtilitiesModule.cpp
new file mode 100644
--- /dev/null
+++ b/Engine/Plugins/Editor/WorldPartitionHLODUtilities/Source/Private/WorldPartitionHLODUtilitiesModule.cpp
@@ -0,0 +1,31 @@
+// Copyright Epic Games, Inc. All Rights Reserved.
+
+#include "WorldPartitionHLODUtilitiesModule.h"
+#include "Modules/ModuleManager.h"
+#include "WorldPartitionHLODUtilities.h"
+
+IMPLEMENT_MODULE(FWorldPartitionHLODUtilitiesModule, WorldPartitionHLODUtilities);
+
+void FWorldPartitionHLODUtilitiesModule::StartupModule()
+{
+	Utilities = nullptr;
+}
+
+void FWorldPartitionHLODUtilitiesModule::ShutdownModule()
+{
+	if (Utilities)
+	{
+		delete Utilities;
+		Utilities = nullptr;
+	}
+}
+
+IWorldPartitionHLODUtilities* FWorldPartitionHLODUtilitiesModule::GetUtilities()
+{
+	if (Utilities == nullptr)
+	{
+		Utilities = new FWorldPartitionHLODUtilities();
+	}
+
+	return Utilities;
+}
diff --git a/Engine/Plugins/Editor/WorldPartitionHLODUtilities/Source/Public/HLODBuilderInstancing.h b/Engine/Plugins/Editor/WorldPartitionHLODUtilities/Source/Public/HLODBuilderInstancing.h
new file mode 100644
--- /dev/null
+++ b/Engine/Plugins/Editor/WorldPartitionHLODUtilities/Source/Public/HLODBuilderInstancing.h
@@ -0,0 +1,22 @@
+// Copyright Epic Games, Inc. All Rights Reserved.
+
+#pragma once
+
+#include "WorldPartition/HLOD/HLODBuilder.h"
+#include "HLODBuilderInstancing.generated.h"
+
+
+/**
+ * Build a AWorldPartitionHLOD whose components are ISMC
+ */
+UCLASS()
+class WORLDPARTITIONHLODUTILITIES_API UHLODBuilderInstancing : public UHLODBuilder
+{
+	 GENERATED_UCLASS_BODY()
+
+public:
+	virtual bool RequiresCompiledAssets() const override { return false; }
+	virtual bool RequiresWarmup() const override { return false; }
+
+	virtual TArray<UPrimitiveComponent*> CreateComponents(AWorldPartitionHLOD* InHLODActor, const UHLODLayer* InHLODLayer, const TArray<UPrimitiveComponent*>& InSubComponents) const override;
+};
diff --git a/Engine/Plugins/Editor/WorldPartitionHLODUtilities/Source/Public/HLODBuilderMeshApproximate.h b/Engine/Plugins/Editor/WorldPartitionHLODUtilities/Source/Public/HLODBuilderMeshApproximate.h
new file mode 100644
--- /dev/null
+++ b/Engine/Plugins/Editor/WorldPartitionHLODUtilities/Source/Public/HLODBuilderMeshApproximate.h
@@ -0,0 +1,41 @@
+// Copyright Epic Games, Inc. All Rights Reserved.
+
+#pragma once
+
+#include "WorldPartition/HLOD/HLODBuilder.h"
+#include "Engine/MeshMerging.h"
+
+#include "HLODBuilderMeshApproximate.generated.h"
+
+class UMaterial;
+
+
+UCLASS(Blueprintable, Config = Engine, PerObjectConfig)
+class WORLDPARTITIONHLODUTILITIES_API UHLODBuilderMeshApproximateSettings : public UHLODBuilderSettings
+{
+	GENERATED_UCLASS_BODY()
+
+	virtual uint32 GetCRC() const override;
+
+	/** Mesh approximation settings */
+	UPROPERTY(EditAnywhere, Config, Category = HLOD, meta=(EditInline))
+	FMeshApproximationSettings MeshApproximationSettings;
+
+	/** Material that will be used by the generated HLOD static mesh */
+	UPROPERTY(EditAnywhere, Config, AdvancedDisplay, Category = HLOD)
+	TSoftObjectPtr<UMaterial> HLODMaterial;
+};
+
+
+/**
+ * Build an approximated mesh using geometry from the provided actors
+ */
+UCLASS()
+class WORLDPARTITIONHLODUTILITIES_API UHLODBuilderMeshApproximate : public UHLODBuilder
+{
+	GENERATED_UCLASS_BODY()
+
+public:
+	virtual UHLODBuilderSettings* CreateSettings(UHLODLayer* InHLODLayer) const override;
+	virtual TArray<UPrimitiveComponent*> CreateComponents(AWorldPartitionHLOD* InHLODActor, const UHLODLayer* InHLODLayer, const TArray<UPrimitiveComponent*>& InSubComponents) const override;
+};
diff --git a/Engine/Plugins/Editor/WorldPartitionHLODUtilities/Source/Public/HLODBuilderMeshMerge.h b/Engine/Plugins/Editor/WorldPartitionHLODUtilities/Source/Public/HLODBuilderMeshMerge.h
new file mode 100644
--- /dev/null
+++ b/Engine/Plugins/Editor/WorldPartitionHLODUtilities/Source/Public/HLODBuilderMeshMerge.h
@@ -0,0 +1,41 @@
+// Copyright Epic Games, Inc. All Rights Reserved.
+
+#pragma once
+
+#include "WorldPartition/HLOD/HLODBuilder.h"
+#include "Engine/MeshMerging.h"
+
+#include "HLODBuilderMeshMerge.generated.h"
+
+class UMaterial;
+
+
+UCLASS(Blueprintable, Config = Engine, PerObjectConfig)
+class WORLDPARTITIONHLODUTILITIES_API UHLODBuilderMeshMergeSettings : public UHLODBuilderSettings
+{
+	GENERATED_UCLASS_BODY()
+
+	virtual uint32 GetCRC() const override;
+
+	/** Merged mesh generation settings */
+	UPROPERTY(EditAnywhere, Config, Category = HLOD)
+	FMeshMergingSettings MeshMergeSettings;
+
+	/** Material that will be used by the generated HLOD static mesh */
+	UPROPERTY(EditAnywhere, Config, AdvancedDisplay, Category = HLOD)
+	TSoftObjectPtr<UMaterial> HLODMaterial;
+};
+
+
+/**
+ * Build a merged mesh using geometry from the provided actors
+ */
+UCLASS()
+class WORLDPARTITIONHLODUTILITIES_API UHLODBuilderMeshMerge : public UHLODBuilder
+{
+	GENERATED_UCLASS_BODY()
+
+public:
+	virtual UHLODBuilderSettings* CreateSettings(UHLODLayer* InHLODLayer) const override;
+	virtual TArray<UPrimitiveComponent*> CreateComponents(AWorldPartitionHLOD* InHLODActor, const UHLODLayer* InHLODLayer, const TArray<UPrimitiveComponent*>& InSubComponents) const override;
+};
diff --git a/Engine/Plugins/Editor/WorldPartitionHLODUtilities/Source/Public/HLODBuilderMeshSimplify.h b/Engine/Plugins/Editor/WorldPartitionHLODUtilities/Source/Public/HLODBuilderMeshSimplify.h
new file mode 100644
--- /dev/null
+++ b/Engine/Plugins/Editor/WorldPartitionHLODUtilities/Source/Public/HLODBuilderMeshSimplify.h
@@ -0,0 +1,41 @@
+// Copyright Epic Games, Inc. All Rights Reserved.
+
+#pragma once
+
+#include "WorldPartition/HLOD/HLODBuilder.h"
+#include "Engine/MeshMerging.h"
+
+#include "HLODBuilderMeshSimplify.generated.h"
+
+class UMaterial;
+
+
+UCLASS(Blueprintable, Config = Engine, PerObjectConfig)
+class WORLDPARTITIONHLODUTILITIES_API UHLODBuilderMeshSimplifySettings : public UHLODBuilderSettings
+{
+	GENERATED_UCLASS_BODY()
+
+	virtual uint32 GetCRC() const override;
+
+	/** Simplified mesh generation settings */
+	UPROPERTY(EditAnywhere, Config, Category = HLOD)
+	FMeshProxySettings MeshSimplifySettings;
+
+	/** Material that will be used by the generated HLOD static mesh */
+	UPROPERTY(EditAnywhere, Config, AdvancedDisplay, Category = HLOD)
+	TSoftObjectPtr<UMaterial> HLODMaterial;
+};
+
+
+/**
+ * Build a simplified mesh using geometry from the provided actors
+ */
+UCLASS()
+class WORLDPARTITIONHLODUTILITIES_API UHLODBuilderMeshSimplify : public UHLODBuilder
+{
+	GENERATED_UCLASS_BODY()
+
+public:
+	virtual UHLODBuilderSettings* CreateSettings(UHLODLayer* InHLODLayer) const override;
+	virtual TArray<UPrimitiveComponent*> CreateComponents(AWorldPartitionHLOD* InHLODActor, const UHLODLayer* InHLODLayer, const TArray<UPrimitiveComponent*>& InSubComponents) const override;
+};
diff --git a/Engine/Plugins/Editor/WorldPartitionHLODUtilities/Source/Public/WorldPartitionHLODUtilities.h b/Engine/Plugins/Editor/WorldPartitionHLODUtilities/Source/Public/WorldPartitionHLODUtilities.h
new file mode 100644
--- /dev/null
+++ b/Engine/Plugins/Editor/WorldPartitionHLODUtilities/Source/Public/WorldPartitionHLODUtilities.h
@@ -0,0 +1,17 @@
+// Copyright Epic Games, Inc. All Rights Reserved.
+
+#pragma once
+
+#include "WorldPartition/HLOD/IWorldPartitionHLODUtilities.h"
+
+/**
+* FWorldPartitionHLODUtilities implementation
+*/
+class FWorldPartitionHLODUtilities : public IWorldPartitionHLODUtilities
+{
+public:
+	virtual TArray<AWorldPartitionHLOD*> CreateHLODActors(FHLODCreationContext& InCreationContext, const FHLODCreationParams& InCreationParams, const TSet<FActorInstance>& InActors, const TArray<const UDataLayer*>& InDataLayers) override;
+	virtual uint32 BuildHLOD(AWorldPartitionHLOD* InHLODActor) override;
+	virtual TSubclassOf<UHLODBuilder> GetHLODBuilderClass(const UHLODLayer* InHLODLayer) override;
+	virtual UHLODBuilderSettings* CreateHLODBuilderSettings(UHLODLayer* InHLODLayer) override;
+};
diff --git a/Engine/Plugins/Editor/WorldPartitionHLODUtilities/Source/Public/WorldPartitionHLODUtilitiesModule.h b/Engine/Plugins/Editor/WorldPartitionHLODUtilities/Source/Public/WorldPartitionHLODUtilitiesModule.h
new file mode 100644
--- /dev/null
+++ b/Engine/Plugins/Editor/WorldPartitionHLODUtilities/Source/Public/WorldPartitionHLODUtilitiesModule.h
@@ -0,0 +1,20 @@
+// Copyright Epic Games, Inc. All Rights Reserved.
+
+#pragma once
+
+#include "WorldPartition/HLOD/IWorldPartitionHLODUtilitiesModule.h"
+
+/**
+* IWorldPartitionHLODUtilities module interface
+*/
+class WORLDPARTITIONHLODUTILITIES_API FWorldPartitionHLODUtilitiesModule : public IWorldPartitionHLODUtilitiesModule
+{
+public:
+	virtual void ShutdownModule() override;
+	virtual void StartupModule() override;
+
+	virtual IWorldPartitionHLODUtilities* GetUtilities() override;
+
+private:
+	IWorldPartitionHLODUtilities* Utilities;
+};
diff --git a/Engine/Plugins/Editor/WorldPartitionHLODUtilities/Source/WorldPartitionHLODUtilities.Build.cs b/Engine/Plugins/Editor/WorldPartitionHLODUtilities/Source/WorldPartitionHLODUtilities.Build.cs
new file mode 100644
--- /dev/null
+++ b/Engine/Plugins/Editor/WorldPartitionHLODUtilities/Source/WorldPartitionHLODUtilities.Build.cs
@@ -0,0 +1,43 @@
+// Copyright Epic Games, Inc. All Rights Reserved.
+using UnrealBuildTool;
+
+public class WorldPartitionHLODUtilities : ModuleRules
+{
+    public WorldPartitionHLODUtilities(ReadOnlyTargetRules Target) : base(Target)
+	{
+        PublicIncludePaths.Add("Developer/WorldPartitionHLODUtilities/Public");
+
+        PublicDependencyModuleNames.AddRange(
+            new string[]
+			{
+				"Core",
+				"CoreUObject"
+			}
+        );
+
+        PrivateDependencyModuleNames.AddRange(
+            new string[]
+			{
+				"Engine",
+				"MaterialUtilities"
+			}
+        );
+
+        PrivateIncludePathModuleNames.AddRange(
+            new string[]
+            {
+				"GeometryProcessingInterfaces",
+			}
+        );
+
+        DynamicallyLoadedModuleNames.AddRange(
+            new string[]
+            {
+                "MeshUtilities",
+                "MeshMergeUtilities",
+                "MeshReductionInterface",
+				"GeometryProcessingInterfaces",
+			}
+        );
+	}
+}
diff --git a/Engine/Plugins/Editor/WorldPartitionHLODUtilities/WorldPartitionHLODUtilities.uplugin b/Engine/Plugins/Editor/WorldPartitionHLODUtilities/WorldPartitionHLODUtilities.uplugin
new file mode 100644
--- /dev/null
+++ b/Engine/Plugins/Editor/WorldPartitionHLODUtilities/WorldPartitionHLODUtilities.uplugin
@@ -0,0 +1,22 @@
+{
+	"FileVersion" : 3,
+	"Version" : 1,
+	"VersionName" : "1.0",
+	"FriendlyName" : "World Partition HLOD Utilities",
+	"Description" : "Editor utility classes & HLOD layer asset types",
+	"Category" : "Editor",
+	"CreatedBy" : "Epic Games, Inc.",
+	"CreatedByURL" : "http://epicgames.com",
+	"EnabledByDefault" : true,
+	"CanContainContent" : false,
+	"IsBetaVersion" : false,
+	"Installed" : false,
+	"Modules" :
+	[
+		{
+			"Name" : "WorldPartitionHLODUtilities",
+			"Type" : "Editor",
+			"LoadingPhase" : "Default"
+		}
+	]
+}
\ No newline at end of file
diff --git a/Engine/Source/Developer/WorldPartitionHLODUtilities/Private/HLODBuilderInstancing.cpp b/Engine/Source/Developer/WorldPartitionHLODUtilities/Private/HLODBuilderInstancing.cpp
deleted file mode 100644
diff --git a/Engine/Source/Developer/WorldPartitionHLODUtilities/Private/HLODBuilderMeshApproximate.cpp b/Engine/Source/Developer/WorldPartitionHLODUtilities/Private/HLODBuilderMeshApproximate.cpp
deleted file mode 100644
diff --git a/Engine/Source/Developer/WorldPartitionHLODUtilities/Private/HLODBuilderMeshMerge.cpp b/Engine/Source/Developer/WorldPartitionHLODUtilities/Private/HLODBuilderMeshMerge.cpp
deleted file mode 100644
diff --git a/Engine/Source/Developer/WorldPartitionHLODUtilities/Private/HLODBuilderMeshSimplify.cpp b/Engine/Source/Developer/WorldPartitionHLODUtilities/Private/HLODBuilderMeshSimplify.cpp
deleted file mode 100644
diff --git a/Engine/Source/Developer/WorldPartitionHLODUtilities/Private/WorldPartitionHLODUtilities.cpp b/Engine/Source/Developer/WorldPartitionHLODUtilities/Private/WorldPartitionHLODUtilities.cpp
deleted file mode 100644
diff --git a/Engine/Source/Developer/WorldPartitionHLODUtilities/Private/WorldPartitionHLODUtilitiesModule.cpp b/Engine/Source/Developer/WorldPartitionHLODUtilities/Private/WorldPartitionHLODUtilitiesModule.cpp
deleted file mode 100644
diff --git a/Engine/Source/Developer/WorldPartitionHLODUtilities/Public/HLODBuilderInstancing.h b/Engine/Source/Developer/WorldPartitionHLODUtilities/Public/HLODBuilderInstancing.h
deleted file mode 100644
diff --git a/Engine/Source/Developer/WorldPartitionHLODUtilities/Public/HLODBuilderMeshApproximate.h b/Engine/Source/Developer/WorldPartitionHLODUtilities/Public/HLODBuilderMeshApproximate.h
deleted file mode 100644
diff --git a/Engine/Source/Developer/WorldPartitionHLODUtilities/Public/HLODBuilderMeshMerge.h b/Engine/Source/Developer/WorldPartitionHLODUtilities/Public/HLODBuilderMeshMerge.h
deleted file mode 100644
diff --git a/Engine/Source/Developer/WorldPartitionHLODUtilities/Public/HLODBuilderMeshSimplify.h b/Engine/Source/Developer/WorldPartitionHLODUtilities/Public/HLODBuilderMeshSimplify.h
deleted file mode 100644
diff --git a/Engine/Source/Developer/WorldPartitionHLODUtilities/Public/IWorldPartitionHLODUtilities.h b/Engine/Source/Developer/WorldPartitionHLODUtilities/Public/IWorldPartitionHLODUtilities.h
deleted file mode 100644
diff --git a/Engine/Source/Developer/WorldPartitionHLODUtilities/Public/WorldPartitionHLODUtilities.h b/Engine/Source/Developer/WorldPartitionHLODUtilities/Public/WorldPartitionHLODUtilities.h
deleted file mode 100644
diff --git a/Engine/Source/Developer/WorldPartitionHLODUtilities/Public/WorldPartitionHLODUtilitiesModule.h b/Engine/Source/Developer/WorldPartitionHLODUtilities/Public/WorldPartitionHLODUtilitiesModule.h
deleted file mode 100644
diff --git a/Engine/Source/Developer/WorldPartitionHLODUtilities/WorldPartitionHLODUtilities.Build.cs b/Engine/Source/Developer/WorldPartitionHLODUtilities/WorldPartitionHLODUtilities.Build.cs
deleted file mode 100644
diff --git a/Engine/Source/Editor/UnrealEd/Private/EditorEngine.cpp b/Engine/Source/Editor/UnrealEd/Private/EditorEngine.cpp
--- a/Engine/Source/Editor/UnrealEd/Private/EditorEngine.cpp
+++ b/Engine/Source/Editor/UnrealEd/Private/EditorEngine.cpp
@@ -1054,7 +1054,6 @@ void UEditorEngine::Init(IEngineLoop* InEngineLoop)
 			TEXT("StaticMeshEditor"),
 			TEXT("EditorFramework"),
 			TEXT("WorldPartitionEditor"),
-			TEXT("WorldPartitionHLODUtilities"),
 			TEXT("EditorConfig"),
 			TEXT("DerivedDataEditor"),
 			TEXT("CSVtoSVG"),
diff --git a/Engine/Source/Editor/UnrealEd/UnrealEd.Build.cs b/Engine/Source/Editor/UnrealEd/UnrealEd.Build.cs
--- a/Engine/Source/Editor/UnrealEd/UnrealEd.Build.cs
+++ b/Engine/Source/Editor/UnrealEd/UnrealEd.Build.cs
@@ -263,7 +263,6 @@ public class UnrealEd : ModuleRules
 				"Media",
 				"VirtualTexturingEditor",
 				"WorldPartitionEditor",
-				"WorldPartitionHLODUtilities",
 				"CSVtoSVG",
 			}
 		);
diff --git a/Engine/Source/Runtime/Engine/Engine.Build.cs b/Engine/Source/Runtime/Engine/Engine.Build.cs
--- a/Engine/Source/Runtime/Engine/Engine.Build.cs
+++ b/Engine/Source/Runtime/Engine/Engine.Build.cs
@@ -244,9 +244,6 @@ public class Engine : ModuleRules
 			
 			PublicIncludePathModuleNames.Add("AnimationDataController");
 			DynamicallyLoadedModuleNames.Add("AnimationDataController");
-			
-			PrivateIncludePathModuleNames.Add("WorldPartitionHLODUtilities");
-			DynamicallyLoadedModuleNames.Add("WorldPartitionHLODUtilities");
 
 			PrivateDependencyModuleNames.AddRange(
 				new string[] {
diff --git a/Engine/Source/Runtime/Engine/Private/WorldPartition/HLOD/HLODActor.cpp b/Engine/Source/Runtime/Engine/Private/WorldPartition/HLOD/HLODActor.cpp
--- a/Engine/Source/Runtime/Engine/Private/WorldPartition/HLOD/HLODActor.cpp
+++ b/Engine/Source/Runtime/Engine/Private/WorldPartition/HLOD/HLODActor.cpp
@@ -12,10 +12,9 @@
 #include "WorldPartition/WorldPartitionActorDesc.h"
 #include "WorldPartition/HLOD/HLODActorDesc.h"
 #include "WorldPartition/HLOD/HLODLayer.h"
+#include "WorldPartition/HLOD/IWorldPartitionHLODUtilitiesModule.h"
 
 #include "Modules/ModuleManager.h"
-#include "IWorldPartitionHLODUtilities.h"
-#include "WorldPartitionHLODUtilitiesModule.h"
 #include "Engine/TextureStreamingTypes.h"
 #endif
 
diff --git a/Engine/Source/Runtime/Engine/Private/WorldPartition/HLOD/HLODLayer.cpp b/Engine/Source/Runtime/Engine/Private/WorldPartition/HLOD/HLODLayer.cpp
--- a/Engine/Source/Runtime/Engine/Private/WorldPartition/HLOD/HLODLayer.cpp
+++ b/Engine/Source/Runtime/Engine/Private/WorldPartition/HLOD/HLODLayer.cpp
@@ -15,10 +15,9 @@
 #include "WorldPartition/WorldPartition.h"
 #include "WorldPartition/WorldPartitionActorDesc.h"
 #include "WorldPartition/WorldPartitionActorDescView.h"
+#include "WorldPartition/HLOD/IWorldPartitionHLODUtilitiesModule.h"
 
 #include "Modules/ModuleManager.h"
-#include "IWorldPartitionHLODUtilities.h"
-#include "WorldPartitionHLODUtilitiesModule.h"
 #endif
 
 DEFINE_LOG_CATEGORY_STATIC(LogHLODLayer, Log, All);
diff --git a/Engine/Source/Runtime/Engine/Private/WorldPartition/RuntimeSpatialHash/RuntimeSpatialHashHLOD.cpp b/Engine/Source/Runtime/Engine/Private/WorldPartition/RuntimeSpatialHash/RuntimeSpatialHashHLOD.cpp
--- a/Engine/Source/Runtime/Engine/Private/WorldPartition/RuntimeSpatialHash/RuntimeSpatialHashHLOD.cpp
+++ b/Engine/Source/Runtime/Engine/Private/WorldPartition/RuntimeSpatialHash/RuntimeSpatialHashHLOD.cpp
@@ -9,6 +9,8 @@
 #include "WorldPartition/HLOD/HLODLayer.h"
 #include "WorldPartition/HLOD/HLODActor.h"
 #include "WorldPartition/HLOD/HLODActorDesc.h"
+#include "WorldPartition/HLOD/IWorldPartitionHLODUtilitiesModule.h"
+
 #include "WorldPartition/WorldPartitionHelpers.h"
 
 #include "UObject/GCObjectScopeGuard.h"
@@ -25,8 +27,6 @@
 #include "StaticMeshCompiler.h"
 
 #include "UObject/SavePackage.h"
-#include "IWorldPartitionHLODUtilities.h"
-#include "WorldPartitionHLODUtilitiesModule.h"
 
 
 DEFINE_LOG_CATEGORY_STATIC(LogWorldPartitionRuntimeSpatialHashHLOD, Log, All);
diff --git a/Engine/Source/Runtime/Engine/Public/WorldPartition/HLOD/IWorldPartitionHLODUtilities.h b/Engine/Source/Runtime/Engine/Public/WorldPartition/HLOD/IWorldPartitionHLODUtilities.h
new file mode 100644
--- /dev/null
+++ b/Engine/Source/Runtime/Engine/Public/WorldPartition/HLOD/IWorldPartitionHLODUtilities.h
@@ -0,0 +1,81 @@
+// Copyright Epic Games, Inc. All Rights Reserved.
+
+#pragma once
+
+#include "CoreMinimal.h"
+#include "Containers/Array.h"
+#include "WorldPartition/WorldPartitionHandle.h"
+#include "WorldPartition/DataLayer/DataLayersID.h"
+#include "WorldPartition/WorldPartitionActorCluster.h"
+
+class AActor;
+class UWorldPartition;
+class UHLODLayer;
+class UHLODBuilder;
+class UHLODBuilderSettings;
+class AWorldPartitionHLOD;
+
+struct ENGINE_API FHLODCreationContext
+{
+	TMap<uint64, FWorldPartitionHandle> HLODActorDescs;
+	TArray<FWorldPartitionReference> ActorReferences;
+};
+
+struct ENGINE_API FHLODCreationParams
+{
+	UWorldPartition* WorldPartition;
+
+	// Everything needed to build the cell hash
+	int64 GridIndexX;
+	int64 GridIndexY;
+	int64 GridIndexZ;
+	FDataLayersID DataLayersID;
+	FName HLODLayerName;
+
+	FName CellName;
+	FBox  CellBounds;
+	uint32 HLODLevel;
+};
+
+/**
+ * Tools for building HLODs in WorldPartition
+ */
+class ENGINE_API IWorldPartitionHLODUtilities
+{
+public:
+	virtual ~IWorldPartitionHLODUtilities() {}
+
+	/**
+	 * Create HLOD actors for a given cell
+	 *
+	 * @param	InCreationContext	HLOD creation context object
+	 * @param	InCreationParams	HLOD creation parameters object
+	 * @param	InActors			The actors for which we'll build an HLOD representation
+	 * @param	InDataLayers		The data layers to assign to the newly created HLOD actors
+	 */
+	virtual TArray<AWorldPartitionHLOD*> CreateHLODActors(FHLODCreationContext& InCreationContext, const FHLODCreationParams& InCreationParams, const TSet<FActorInstance>& InActors, const TArray<const UDataLayer*>& InDataLayers) = 0;
+
+	/**
+	 * Build HLOD for the specified AWorldPartitionHLOD actor.
+	 *
+	 * @param 	InHLODActor		The HLOD actor for which we'll build the HLOD
+	 * @return An hash that represent the content used to build this HLOD.
+	 */
+	virtual uint32 BuildHLOD(AWorldPartitionHLOD* InHLODActor) = 0;
+
+	/**
+	 * Retrieve the HLOD Builder class to use for the given HLODLayer.
+	 * 
+	 * @param	InHLODLayer		HLODLayer
+	 * @return The HLOD builder subclass to use for building HLODs for the provided HLOD layer.
+	 */
+	virtual TSubclassOf<UHLODBuilder> GetHLODBuilderClass(const UHLODLayer* InHLODLayer) = 0;
+
+	/**
+	 * Create the HLOD builder settings for the provided HLOD layer object. The type of settings created will depend on the HLOD layer type.
+	 *
+	 * @param 	InHLODLayer		The HLOD layer for which we'll create a setting object
+	 * @return A newly created UHLODBuilderSettings object, outered to the provided HLOD layer.
+	 */
+	virtual UHLODBuilderSettings* CreateHLODBuilderSettings(UHLODLayer* InHLODLayer) = 0;
+};
diff --git a/Engine/Source/Runtime/Engine/Public/WorldPartition/HLOD/IWorldPartitionHLODUtilitiesModule.h b/Engine/Source/Runtime/Engine/Public/WorldPartition/HLOD/IWorldPartitionHLODUtilitiesModule.h
new file mode 100644
--- /dev/null
+++ b/Engine/Source/Runtime/Engine/Public/WorldPartition/HLOD/IWorldPartitionHLODUtilitiesModule.h
@@ -0,0 +1,12 @@
+// Copyright Epic Games, Inc. All Rights Reserved.
+
+#pragma once
+
+#include "Modules/ModuleInterface.h"
+#include "WorldPartition/HLOD/IWorldPartitionHLODUtilities.h"
+
+class ENGINE_API IWorldPartitionHLODUtilitiesModule : public IModuleInterface
+{
+public:
+	virtual IWorldPartitionHLODUtilities* GetUtilities() = 0;
+};
