Fixed warning in WorldPartitionHLODUtilities.Build.cs, introduced in CL 18584312
#rnx

diff --git a/Engine/Plugins/Editor/WorldPartitionHLODUtilities/Source/WorldPartitionHLODUtilities.Build.cs b/Engine/Plugins/Editor/WorldPartitionHLODUtilities/Source/WorldPartitionHLODUtilities.Build.cs
--- a/Engine/Plugins/Editor/WorldPartitionHLODUtilities/Source/WorldPartitionHLODUtilities.Build.cs
+++ b/Engine/Plugins/Editor/WorldPartitionHLODUtilities/Source/WorldPartitionHLODUtilities.Build.cs
@@ -5,8 +5,6 @@ public class WorldPartitionHLODUtilities : ModuleRules
 {
     public WorldPartitionHLODUtilities(ReadOnlyTargetRules Target) : base(Target)
 	{
-        PublicIncludePaths.Add("Developer/WorldPartitionHLODUtilities/Public");
-
         PublicDependencyModuleNames.AddRange(
             new string[]
 			{
