WP HLOD - Add Landscape support (WIP)

Refactored the HLOD build process to allow more flexibility in the type of components that are used as input

* Used to support only Primitive components as input, can now be any type of component.
* Classes of component can override GetCustomHLODBuilderClass()
* HLODBuilder class as no more dependency on WorldPartition (no more need for HLOD layer, HLOD actor), to be moved later

This allows the HLOD system to include whatever is needed to represent a component at a distance.
For example, this is used to build custom HLODs for landscape streaming proxies.
This could be used to generate fake lights, represent custom FX at a distance, or even to insert externally generated HLODs.

diff --git a/Engine/Plugins/Editor/WorldPartitionHLODUtilities/Source/Private/HLODBuilderInstancing.cpp b/Engine/Plugins/Editor/WorldPartitionHLODUtilities/Source/Private/HLODBuilderInstancing.cpp
--- a/Engine/Plugins/Editor/WorldPartitionHLODUtilities/Source/Private/HLODBuilderInstancing.cpp
+++ b/Engine/Plugins/Editor/WorldPartitionHLODUtilities/Source/Private/HLODBuilderInstancing.cpp
@@ -45,31 +45,28 @@ UHLODBuilderInstancing::UHLODBuilderInstancing(const FObjectInitializer& ObjectI
 {
 }
 
-TArray<UPrimitiveComponent*> UHLODBuilderInstancing::CreateComponents(AWorldPartitionHLOD* InHLODActor, const UHLODLayer* InHLODLayer, const TArray<UPrimitiveComponent*>& InSubComponents) const
+TArray<UActorComponent*> UHLODBuilderInstancing::Build(const FHLODBuildContext& InHLODBuildContext, const TArray<UActorComponent*>& InSourceComponents) const
 {
-	TRACE_CPUPROFILER_EVENT_SCOPE(UHLODBuilderInstancing::CreateComponents);
-
-	TArray<UPrimitiveComponent*> Components;
+	TRACE_CPUPROFILER_EVENT_SCOPE(UHLODBuilderInstancing::Build);
 
+	TArray<UStaticMeshComponent*> SourceStaticMeshComponents = FilterComponents<UStaticMeshComponent>(InSourceComponents);
+	
 	// Prepare instance batches
 	TMap<FISMComponentDescriptor, FInstancingData> InstancesData;
-	for (UPrimitiveComponent* Primitive : InSubComponents)
+	for (UStaticMeshComponent* SMC : SourceStaticMeshComponents)
 	{
-		if (UStaticMeshComponent* SMC = Cast<UStaticMeshComponent>(Primitive))
-		{
-			FCustomISMComponentDescriptor ISMComponentDescriptor(SMC);
-			FInstancingData& InstancingData = InstancesData.FindOrAdd(ISMComponentDescriptor);
+		FCustomISMComponentDescriptor ISMComponentDescriptor(SMC);
+		FInstancingData& InstancingData = InstancesData.FindOrAdd(ISMComponentDescriptor);
 
-			if (UInstancedStaticMeshComponent* ISMC = Cast<UInstancedStaticMeshComponent>(SMC))
-			{
-				InstancingData.NumCustomDataFloats = FMath::Max(InstancingData.NumCustomDataFloats, ISMC->NumCustomDataFloats);
-				InstancingData.RandomSeeds.Add( { InstancingData.NumInstances, ISMC->InstancingRandomSeed } );
-				InstancingData.NumInstances += ISMC->GetInstanceCount();
-			}
-			else
-			{
-				InstancingData.NumInstances++;
-			}
+		if (UInstancedStaticMeshComponent* ISMC = Cast<UInstancedStaticMeshComponent>(SMC))
+		{
+			InstancingData.NumCustomDataFloats = FMath::Max(InstancingData.NumCustomDataFloats, ISMC->NumCustomDataFloats);
+			InstancingData.RandomSeeds.Add( { InstancingData.NumInstances, ISMC->InstancingRandomSeed } );
+			InstancingData.NumInstances += ISMC->GetInstanceCount();
+		}
+		else
+		{
+			InstancingData.NumInstances++;
 		}
 	}
 
@@ -84,36 +81,34 @@ TArray<UPrimitiveComponent*> UHLODBuilderInstancing::CreateComponents(AWorldPart
 	}
 	
 	// Append all transforms & per instance custom data
-	for (UPrimitiveComponent* Primitive : InSubComponents)
+	for (UStaticMeshComponent* SMC : SourceStaticMeshComponents)
 	{
-		if (UStaticMeshComponent* SMC = Cast<UStaticMeshComponent>(Primitive))
-		{
-			FCustomISMComponentDescriptor ISMComponentDescriptor(SMC);
-			FInstancingData& InstancingData = InstancesData.FindChecked(ISMComponentDescriptor);
+		FCustomISMComponentDescriptor ISMComponentDescriptor(SMC);
+		FInstancingData& InstancingData = InstancesData.FindChecked(ISMComponentDescriptor);
 
-			if (UInstancedStaticMeshComponent* ISMC = Cast<UInstancedStaticMeshComponent>(SMC))
-			{
-				// Add transforms
-				for (int32 InstanceIdx = 0; InstanceIdx < ISMC->GetInstanceCount(); InstanceIdx++)
-				{
-					FTransform& InstanceTransform = InstancingData.InstancesTransforms.AddDefaulted_GetRef();
-					ISMC->GetInstanceTransform(InstanceIdx, InstanceTransform, true);
-				}
-
-				// Add per instance custom data
-				int32 NumCustomDataFloatToAdd = ISMC->GetInstanceCount() * InstancingData.NumCustomDataFloats;
-				InstancingData.InstancesCustomData.Append(ISMC->PerInstanceSMCustomData);
-				InstancingData.InstancesCustomData.AddDefaulted(NumCustomDataFloatToAdd - ISMC->PerInstanceSMCustomData.Num());
-			}
-			else
+		if (UInstancedStaticMeshComponent* ISMC = Cast<UInstancedStaticMeshComponent>(SMC))
+		{
+			// Add transforms
+			for (int32 InstanceIdx = 0; InstanceIdx < ISMC->GetInstanceCount(); InstanceIdx++)
 			{
-				InstancingData.InstancesTransforms.Add(SMC->GetComponentTransform());
-				InstancingData.InstancesCustomData.AddDefaulted(InstancingData.NumCustomDataFloats);
+				FTransform& InstanceTransform = InstancingData.InstancesTransforms.AddDefaulted_GetRef();
+				ISMC->GetInstanceTransform(InstanceIdx, InstanceTransform, true);
 			}
+
+			// Add per instance custom data
+			int32 NumCustomDataFloatToAdd = ISMC->GetInstanceCount() * InstancingData.NumCustomDataFloats;
+			InstancingData.InstancesCustomData.Append(ISMC->PerInstanceSMCustomData);
+			InstancingData.InstancesCustomData.AddDefaulted(NumCustomDataFloatToAdd - ISMC->PerInstanceSMCustomData.Num());
+		}
+		else
+		{
+			InstancingData.InstancesTransforms.Add(SMC->GetComponentTransform());
+			InstancingData.InstancesCustomData.AddDefaulted(InstancingData.NumCustomDataFloats);
 		}
 	}
 
 	// Create an ISMC for each SM asset we found
+	TArray<UActorComponent*> HLODComponents;
 	for (auto& Entry : InstancesData)
 	{
 		const FISMComponentDescriptor& ISMComponentDescriptor = Entry.Key;
@@ -121,7 +116,7 @@ TArray<UPrimitiveComponent*> UHLODBuilderInstancing::CreateComponents(AWorldPart
 
 		check(EntryInstancingData.InstancesTransforms.Num() * EntryInstancingData.NumCustomDataFloats == EntryInstancingData.InstancesCustomData.Num());
 
-		UInstancedStaticMeshComponent* Component = ISMComponentDescriptor.CreateComponent(InHLODActor);
+		UInstancedStaticMeshComponent* Component = ISMComponentDescriptor.CreateComponent(GetTransientPackage());
 		Component->SetForcedLodModel(Component->GetStaticMesh()->GetNumLODs());
 		Component->NumCustomDataFloats = EntryInstancingData.NumCustomDataFloats;
 		Component->AddInstances(EntryInstancingData.InstancesTransforms, /*bShouldReturnIndices*/false, /*bWorldSpace*/true);
@@ -137,8 +132,8 @@ TArray<UPrimitiveComponent*> UHLODBuilderInstancing::CreateComponents(AWorldPart
 			Component->AdditionalRandomSeeds = TArrayView<FInstancedStaticMeshRandomSeed>(&EntryInstancingData.RandomSeeds[1], EntryInstancingData.RandomSeeds.Num() - 1);
 		}
 
-		Components.Add(Component);
+		HLODComponents.Add(Component);
 	};
 
-	return Components;
+	return HLODComponents;
 }
diff --git a/Engine/Plugins/Editor/WorldPartitionHLODUtilities/Source/Private/HLODBuilderMeshApproximate.cpp b/Engine/Plugins/Editor/WorldPartitionHLODUtilities/Source/Private/HLODBuilderMeshApproximate.cpp
--- a/Engine/Plugins/Editor/WorldPartitionHLODUtilities/Source/Private/HLODBuilderMeshApproximate.cpp
+++ b/Engine/Plugins/Editor/WorldPartitionHLODUtilities/Source/Private/HLODBuilderMeshApproximate.cpp
@@ -64,49 +64,38 @@ uint32 UHLODBuilderMeshApproximateSettings::GetCRC() const
 	return Hash;
 }
 
-UHLODBuilderSettings* UHLODBuilderMeshApproximate::CreateSettings(UHLODLayer* InHLODLayer) const 
+TSubclassOf<UHLODBuilderSettings> UHLODBuilderMeshApproximate::GetSettingsClass() const
 {
-	UHLODBuilderMeshApproximateSettings* HLODBuilderSettings = NewObject<UHLODBuilderMeshApproximateSettings>(InHLODLayer);
-
-	// If previous settings object is null, this means we have an older version of the object. Populate with the deprecated settings.
-	if (InHLODLayer->GetHLODBuilderSettings() == nullptr)
-	{
-		HLODBuilderSettings->MeshApproximationSettings = InHLODLayer->MeshApproximationSettings_DEPRECATED;
-		HLODBuilderSettings->HLODMaterial = InHLODLayer->HLODMaterial_DEPRECATED;
-	}
-
-	return HLODBuilderSettings;
+	return UHLODBuilderMeshApproximateSettings::StaticClass();
 }
 
-TArray<UPrimitiveComponent*> UHLODBuilderMeshApproximate::CreateComponents(AWorldPartitionHLOD* InHLODActor, const UHLODLayer* InHLODLayer, const TArray<UPrimitiveComponent*>& InSubComponents) const
+TArray<UActorComponent*> UHLODBuilderMeshApproximate::Build(const FHLODBuildContext& InHLODBuildContext, const TArray<UActorComponent*>& InSourceComponents) const
 {
-	TRACE_CPUPROFILER_EVENT_SCOPE(UHLODBuilderMeshApproximate::CreateComponents);
+	TRACE_CPUPROFILER_EVENT_SCOPE(UHLODBuilderMeshApproximate::Build);
 
-	TSet<AActor*> Actors;
-	TArray<UPrimitiveComponent*> InstancedComponents;
+	TArray<UStaticMeshComponent*> StaticMeshComponents;
+	TArray<UActorComponent*> InstancedComponents;
 
 	// Filter the input components
-	for (UPrimitiveComponent* SubComponent : InSubComponents)
+	for (UStaticMeshComponent* SubComponent : FilterComponents<UStaticMeshComponent>(InSourceComponents))
 	{
-		if (!SubComponent)
-		{
-			continue;
-		}
-
 		switch (SubComponent->HLODBatchingPolicy)
 		{
 		case EHLODBatchingPolicy::None:
-			Actors.Add(SubComponent->GetOwner());
+			StaticMeshComponents.Add(SubComponent);
 			break;
 		case EHLODBatchingPolicy::Instancing:
 			InstancedComponents.Add(SubComponent);
 			break;
 		case EHLODBatchingPolicy::MeshSection:
 			InstancedComponents.Add(SubComponent);
-			UE_LOG(LogHLODBuilder, Warning, TEXT("EHLODBatchingPolicy::MeshSection is not yet supported by the MeshApproximate builder."));
+			UE_LOG(LogHLODBuilder, Warning, TEXT("EHLODBatchingPolicy::MeshSection is not yet supported by the UHLODBuilderMeshSimplify builder."));
 			break;
 		}
 	}
+
+	TSet<AActor*> Actors;
+	Algo::Transform(StaticMeshComponents, Actors, [](UPrimitiveComponent* PrimitiveComponent) { return PrimitiveComponent->GetOwner(); });
 	
 	IGeometryProcessingInterfacesModule& GeomProcInterfaces = FModuleManager::Get().LoadModuleChecked<IGeometryProcessingInterfacesModule>("GeometryProcessingInterfaces");
 	IGeometryProcessing_ApproximateActors* ApproxActorsAPI = GeomProcInterfaces.GetApproximateActorsImplementation();
@@ -115,12 +104,12 @@ TArray<UPrimitiveComponent*> UHLODBuilderMeshApproximate::CreateComponents(AWorl
 	// Construct options for ApproximateActors operation
 	//
 
-	const UHLODBuilderMeshApproximateSettings* MeshApproximateSettings = CastChecked<UHLODBuilderMeshApproximateSettings>(InHLODLayer->GetHLODBuilderSettings());
+	const UHLODBuilderMeshApproximateSettings* MeshApproximateSettings = CastChecked<UHLODBuilderMeshApproximateSettings>(HLODBuilderSettings);
 	const FMeshApproximationSettings& UseSettings = MeshApproximateSettings->MeshApproximationSettings;
 	UMaterial* HLODMaterial = MeshApproximateSettings->HLODMaterial.LoadSynchronous();
 
 	IGeometryProcessing_ApproximateActors::FOptions Options = ApproxActorsAPI->ConstructOptions(UseSettings);
-	Options.BasePackagePath = InHLODActor->GetPackage()->GetName();
+	Options.BasePackagePath = InHLODBuildContext.AssetsOuter->GetPackage()->GetName();
 	Options.bGenerateLightmapUVs = false;
 	Options.bCreatePhysicsBody = false;
 
@@ -141,7 +130,7 @@ TArray<UPrimitiveComponent*> UHLODBuilderMeshApproximate::CreateComponents(AWorl
 		FBoxSphereBounds Bounds;
 		bool bFirst = true;
 
-		for (UPrimitiveComponent* Component : InSubComponents)
+		for (UPrimitiveComponent* Component : StaticMeshComponents)
 		{
 			FBoxSphereBounds ComponentBounds = Component->Bounds;
 			Bounds = bFirst ? ComponentBounds : Bounds + ComponentBounds;
@@ -179,13 +168,13 @@ TArray<UPrimitiveComponent*> UHLODBuilderMeshApproximate::CreateComponents(AWorl
 	IGeometryProcessing_ApproximateActors::FResults Results;
 	ApproxActorsAPI->ApproximateActors(Actors.Array(), Options, Results);
 
-	TArray<UPrimitiveComponent*> Components;
+	TArray<UActorComponent*> Components;
 	if (Results.ResultCode == IGeometryProcessing_ApproximateActors::EResultCode::Success)
 	{
-		auto FixupAsset = [&InHLODActor](UObject* Asset)
+		auto FixupAsset = [InHLODBuildContext](UObject* Asset)
 		{
 			Asset->ClearFlags(RF_Public | RF_Standalone);
-			Asset->Rename(nullptr, InHLODActor->GetPackage(), REN_NonTransactional | REN_DontCreateRedirectors | REN_ForceNoResetLoaders);
+			Asset->Rename(nullptr, InHLODBuildContext.AssetsOuter, REN_NonTransactional | REN_DontCreateRedirectors | REN_ForceNoResetLoaders);
 		};
 	
 		Algo::ForEach(Results.NewMeshAssets, FixupAsset);
@@ -194,7 +183,7 @@ TArray<UPrimitiveComponent*> UHLODBuilderMeshApproximate::CreateComponents(AWorl
 
 		for (UStaticMesh* StaticMesh : Results.NewMeshAssets)
 		{
-			UStaticMeshComponent* Component = NewObject<UStaticMeshComponent>(InHLODActor);
+			UStaticMeshComponent* Component = NewObject<UStaticMeshComponent>();
 			Component->SetStaticMesh(StaticMesh);
 
 			Components.Add(Component);
@@ -240,7 +229,7 @@ TArray<UPrimitiveComponent*> UHLODBuilderMeshApproximate::CreateComponents(AWorl
 	if (InstancedComponents.Num())
 	{
 		UHLODBuilderInstancing* InstancingHLODBuilder = NewObject<UHLODBuilderInstancing>();
-		Components.Append(InstancingHLODBuilder->CreateComponents(InHLODActor, nullptr, InstancedComponents));
+		Components.Append(InstancingHLODBuilder->Build(InHLODBuildContext, InstancedComponents));
 	}
 
 	return Components;
diff --git a/Engine/Plugins/Editor/WorldPartitionHLODUtilities/Source/Private/HLODBuilderMeshMerge.cpp b/Engine/Plugins/Editor/WorldPartitionHLODUtilities/Source/Private/HLODBuilderMeshMerge.cpp
--- a/Engine/Plugins/Editor/WorldPartitionHLODUtilities/Source/Private/HLODBuilderMeshMerge.cpp
+++ b/Engine/Plugins/Editor/WorldPartitionHLODUtilities/Source/Private/HLODBuilderMeshMerge.cpp
@@ -60,46 +60,39 @@ uint32 UHLODBuilderMeshMergeSettings::GetCRC() const
 	return Hash;
 }
 
-UHLODBuilderSettings* UHLODBuilderMeshMerge::CreateSettings(UHLODLayer* InHLODLayer) const
+TSubclassOf<UHLODBuilderSettings> UHLODBuilderMeshMerge::GetSettingsClass() const
 {
-	UHLODBuilderMeshMergeSettings* HLODBuilderSettings = NewObject<UHLODBuilderMeshMergeSettings>(InHLODLayer);
-
-	// If previous settings object is null, this means we have an older version of the object. Populate with the deprecated settings.
-	if (InHLODLayer->GetHLODBuilderSettings() == nullptr)
-	{
-		HLODBuilderSettings->MeshMergeSettings = InHLODLayer->MeshMergeSettings_DEPRECATED;
-		HLODBuilderSettings->HLODMaterial = InHLODLayer->HLODMaterial_DEPRECATED;
-	}
-
-	return HLODBuilderSettings;
+	return UHLODBuilderMeshMergeSettings::StaticClass();
 }
 
-TArray<UPrimitiveComponent*> UHLODBuilderMeshMerge::CreateComponents(AWorldPartitionHLOD* InHLODActor, const UHLODLayer* InHLODLayer, const TArray<UPrimitiveComponent*>& InSubComponents) const
+TArray<UActorComponent*> UHLODBuilderMeshMerge::Build(const FHLODBuildContext& InHLODBuildContext, const TArray<UActorComponent*>& InSourceComponents) const
 {
 	TRACE_CPUPROFILER_EVENT_SCOPE(UHLODBuilderMeshMerge::CreateComponents);
 
+	TArray<UPrimitiveComponent*> SourcePrimitiveComponents = FilterComponents<UPrimitiveComponent>(InSourceComponents);
+
 	TArray<UObject*> Assets;
 	FVector MergedActorLocation;
 
-	const UHLODBuilderMeshMergeSettings* MeshMergeSettings = CastChecked<UHLODBuilderMeshMergeSettings>(InHLODLayer->GetHLODBuilderSettings());
+	const UHLODBuilderMeshMergeSettings* MeshMergeSettings = CastChecked<UHLODBuilderMeshMergeSettings>(HLODBuilderSettings);
 	const FMeshMergingSettings& UseSettings = MeshMergeSettings->MeshMergeSettings;
 	UMaterial* HLODMaterial = MeshMergeSettings->HLODMaterial.LoadSynchronous();
 
 	const IMeshMergeUtilities& MeshMergeUtilities = FModuleManager::Get().LoadModuleChecked<IMeshMergeModule>("MeshMergeUtilities").GetUtilities();
-	MeshMergeUtilities.MergeComponentsToStaticMesh(InSubComponents, InHLODActor->GetWorld(), UseSettings, HLODMaterial, InHLODActor->GetPackage(), InHLODActor->GetActorLabel(), Assets, MergedActorLocation, 0.25f, false);
+	MeshMergeUtilities.MergeComponentsToStaticMesh(SourcePrimitiveComponents, InHLODBuildContext.World, UseSettings, HLODMaterial, InHLODBuildContext.AssetsOuter->GetPackage(), InHLODBuildContext.AssetsBaseName, Assets, MergedActorLocation, 0.25f, false);
 
 	UStaticMeshComponent* Component = nullptr;
-	Algo::ForEach(Assets, [this, InHLODActor, &Component, &MergedActorLocation](UObject* Asset)
+	Algo::ForEach(Assets, [this, &Component, &MergedActorLocation](UObject* Asset)
 	{
 		Asset->ClearFlags(RF_Public | RF_Standalone);
 
 		if (Cast<UStaticMesh>(Asset))
 		{
-			Component = NewObject<UStaticMeshComponent>(InHLODActor);
+			Component = NewObject<UStaticMeshComponent>();
 			Component->SetStaticMesh(static_cast<UStaticMesh*>(Asset));
 			Component->SetWorldLocation(MergedActorLocation);
 		}
 	});
 
-	return TArray<UPrimitiveComponent*>({ Component });
+	return TArray<UActorComponent*>({ Component });
 }
diff --git a/Engine/Plugins/Editor/WorldPartitionHLODUtilities/Source/Private/HLODBuilderMeshSimplify.cpp b/Engine/Plugins/Editor/WorldPartitionHLODUtilities/Source/Private/HLODBuilderMeshSimplify.cpp
--- a/Engine/Plugins/Editor/WorldPartitionHLODUtilities/Source/Private/HLODBuilderMeshSimplify.cpp
+++ b/Engine/Plugins/Editor/WorldPartitionHLODUtilities/Source/Private/HLODBuilderMeshSimplify.cpp
@@ -63,84 +63,67 @@ uint32 UHLODBuilderMeshSimplifySettings::GetCRC() const
 	return Hash;
 }
 
-UHLODBuilderSettings* UHLODBuilderMeshSimplify::CreateSettings(UHLODLayer* InHLODLayer) const
+TSubclassOf<UHLODBuilderSettings> UHLODBuilderMeshSimplify::GetSettingsClass() const
 {
-	UHLODBuilderMeshSimplifySettings* HLODBuilderSettings = NewObject<UHLODBuilderMeshSimplifySettings>(InHLODLayer);
-
-	// If previous settings object is null, this means we have an older version of the object. Populate with the deprecated settings.
-	if (InHLODLayer->GetHLODBuilderSettings() == nullptr)
-	{
-		HLODBuilderSettings->MeshSimplifySettings = InHLODLayer->MeshSimplifySettings_DEPRECATED;
-		HLODBuilderSettings->HLODMaterial = InHLODLayer->HLODMaterial_DEPRECATED;
-	}
-
-	return HLODBuilderSettings;
+	return UHLODBuilderMeshSimplifySettings::StaticClass();
 }
 
-TArray<UPrimitiveComponent*> UHLODBuilderMeshSimplify::CreateComponents(AWorldPartitionHLOD* InHLODActor, const UHLODLayer* InHLODLayer, const TArray<UPrimitiveComponent*>& InSubComponents) const
+TArray<UActorComponent*> UHLODBuilderMeshSimplify::Build(const FHLODBuildContext& InHLODBuildContext, const TArray<UActorComponent*>& InSourceComponents) const
 {
-	TRACE_CPUPROFILER_EVENT_SCOPE(UHLODBuilderMeshSimplifySettings::CreateComponents);
-
-	TArray<UObject*> Assets;
-	FCreateProxyDelegate ProxyDelegate;
-	ProxyDelegate.BindLambda([&Assets](const FGuid Guid, TArray<UObject*>& InAssetsCreated) { Assets = InAssetsCreated; });
+	TRACE_CPUPROFILER_EVENT_SCOPE(UHLODBuilderMeshSimplifySettings::Build);
 
 	TArray<UStaticMeshComponent*> StaticMeshComponents;
-	TArray<UPrimitiveComponent*> InstancedComponents;
+	TArray<UActorComponent*> InstancedComponents;
 
 	// Filter the input components
-	for (UPrimitiveComponent* SubComponent : InSubComponents)
+	for (UStaticMeshComponent* SubComponent : FilterComponents<UStaticMeshComponent>(InSourceComponents))
 	{
-		if (!SubComponent)
-		{
-			continue;
-		}
-
 		switch (SubComponent->HLODBatchingPolicy)
 		{
 		case EHLODBatchingPolicy::None:
-			if (UStaticMeshComponent* SMC = Cast<UStaticMeshComponent>(SubComponent))
-			{
-				StaticMeshComponents.Add(SMC);
-			}
+			StaticMeshComponents.Add(SubComponent);
 			break;
 		case EHLODBatchingPolicy::Instancing:
 			InstancedComponents.Add(SubComponent);
 			break;
 		case EHLODBatchingPolicy::MeshSection:
 			InstancedComponents.Add(SubComponent);
-			UE_LOG(LogHLODBuilder, Warning, TEXT("EHLODBatchingPolicy::MeshSection is not yet supported by the MeshApproximate builder."));
+			UE_LOG(LogHLODBuilder, Warning, TEXT("EHLODBatchingPolicy::MeshSection is not yet supported by the UHLODBuilderMeshSimplify builder."));
 			break;
 		}
 	}
 
-	const UHLODBuilderMeshSimplifySettings* MeshSimplifySettings = CastChecked<UHLODBuilderMeshSimplifySettings>(InHLODLayer->GetHLODBuilderSettings());
+	const UHLODBuilderMeshSimplifySettings* MeshSimplifySettings = CastChecked<UHLODBuilderMeshSimplifySettings>(HLODBuilderSettings);
 	const FMeshProxySettings& UseSettings = MeshSimplifySettings->MeshSimplifySettings;
 	UMaterial* HLODMaterial = MeshSimplifySettings->HLODMaterial.LoadSynchronous();
 
+	TArray<UObject*> Assets;
+	FCreateProxyDelegate ProxyDelegate;
+	ProxyDelegate.BindLambda([&Assets](const FGuid Guid, TArray<UObject*>& InAssetsCreated) { Assets = InAssetsCreated; });
+
 	const IMeshMergeUtilities& MeshMergeUtilities = FModuleManager::Get().LoadModuleChecked<IMeshMergeModule>("MeshMergeUtilities").GetUtilities();
-	MeshMergeUtilities.CreateProxyMesh(StaticMeshComponents, UseSettings, HLODMaterial, InHLODActor->GetPackage(), InHLODActor->GetActorLabel(), FGuid::NewGuid(), ProxyDelegate, true);
+	MeshMergeUtilities.CreateProxyMesh(StaticMeshComponents, UseSettings, HLODMaterial, InHLODBuildContext.AssetsOuter->GetPackage(), InHLODBuildContext.AssetsBaseName, FGuid::NewGuid(), ProxyDelegate, true);
 
 	UStaticMeshComponent* Component = nullptr;
-	Algo::ForEach(Assets, [this, InHLODActor, &Component](UObject* Asset)
+	Algo::ForEach(Assets, [this, &Component](UObject* Asset)
 	{
 		Asset->ClearFlags(RF_Public | RF_Standalone);
 
 		if (Cast<UStaticMesh>(Asset))
 		{
-			Component = NewObject<UStaticMeshComponent>(InHLODActor);
+			Component = NewObject<UStaticMeshComponent>();
 			Component->SetStaticMesh(static_cast<UStaticMesh*>(Asset));
 		}
 	});
 
-	TArray<UPrimitiveComponent*> Components;
+	TArray<UActorComponent*> Components;
 	Components.Add(Component);
 
 	// Batch instances
 	if (InstancedComponents.Num())
 	{
 		UHLODBuilderInstancing* InstancingHLODBuilder = NewObject<UHLODBuilderInstancing>();
-		Components.Append(InstancingHLODBuilder->CreateComponents(InHLODActor, nullptr, InstancedComponents));
+		Components.Append(InstancingHLODBuilder->Build(InHLODBuildContext, InstancedComponents));
 	}
 
 	return Components;
diff --git a/Engine/Plugins/Editor/WorldPartitionHLODUtilities/Source/Private/WorldPartitionHLODUtilities.cpp b/Engine/Plugins/Editor/WorldPartitionHLODUtilities/Source/Private/WorldPartitionHLODUtilities.cpp
--- a/Engine/Plugins/Editor/WorldPartitionHLODUtilities/Source/Private/WorldPartitionHLODUtilities.cpp
+++ b/Engine/Plugins/Editor/WorldPartitionHLODUtilities/Source/Private/WorldPartitionHLODUtilities.cpp
@@ -18,10 +18,6 @@
 #include "Serialization/ArchiveCrc32.h"
 #include "Templates/UniquePtr.h"
 
-#include "Engine/StaticMesh.h"
-#include "Engine/HLODProxy.h"
-#include "Serialization/ArchiveCrc32.h"
-#include "Materials/Material.h"
 #include "AssetCompilingManager.h"
 
 #include "HLODBuilderInstancing.h"
@@ -83,63 +79,18 @@ static uint32 ComputeHLODHash(AWorldPartitionHLOD* InHLODActor, const TArray<AAc
 
 	// HLOD Layer
 	uint32 HLODLayerHash = GetCRC(InHLODActor->GetSubActorsHLODLayer());
-	UE_LOG(LogHLODBuilder, VeryVerbose, TEXT(" - HLODLayer (%s) = %x"), *InHLODActor->GetSubActorsHLODLayer()->GetName(), HLODLayerHash);
+	UE_LOG(LogHLODBuilder, VeryVerbose, TEXT(" - HLOD Layer (%s) = %x"), *InHLODActor->GetSubActorsHLODLayer()->GetName(), HLODLayerHash);
 	Ar << HLODLayerHash;
 
-	// We get the CRC of each component
-	TArray<uint32> ComponentsCRCs;
-	for (UPrimitiveComponent* Component : UHLODBuilder::GatherPrimitiveComponents(InActors))
-	{
-		if (UStaticMeshComponent* StaticMeshComponent = Cast<UStaticMeshComponent>(Component))
-		{
-			uint32 ComponentCRC = 0;
-
-			UE_LOG(LogHLODBuilder, VeryVerbose, TEXT(" - Component \'%s\' from actor \'%s\'"), *Component->GetName(), *Component->GetOwner()->GetName());
-
-			// CRC component
-			uint32 StaticMeshComponentCRC = UHLODProxy::GetCRC(StaticMeshComponent);
-			UE_LOG(LogHLODBuilder, VeryVerbose, TEXT("     - StaticMeshComponent (%s) = %x"), *StaticMeshComponent->GetName(), StaticMeshComponentCRC);
-			ComponentCRC = HashCombine(ComponentCRC, StaticMeshComponentCRC);
-
-			if (UStaticMesh* StaticMesh = StaticMeshComponent->GetStaticMesh())
-			{
-				// CRC static mesh
-				int32 StaticMeshCRC = UHLODProxy::GetCRC(StaticMesh);
-				UE_LOG(LogHLODBuilder, VeryVerbose, TEXT("     - StaticMesh (%s) = %x"), *StaticMesh->GetName(), StaticMeshCRC);
-				ComponentCRC = HashCombine(ComponentCRC, StaticMeshCRC);
-
-				// CRC materials
-				const int32 NumMaterials = StaticMeshComponent->GetNumMaterials();
-				for (int32 MaterialIndex = 0; MaterialIndex < NumMaterials; ++MaterialIndex)
-				{
-					UMaterialInterface* MaterialInterface = StaticMeshComponent->GetMaterial(MaterialIndex);
-					if (MaterialInterface)
-					{
-						uint32 MaterialInterfaceCRC = UHLODProxy::GetCRC(MaterialInterface);
-						UE_LOG(LogHLODBuilder, VeryVerbose, TEXT("     - MaterialInterface (%s) = %x"), *MaterialInterface->GetName(), MaterialInterfaceCRC);
-						ComponentCRC = HashCombine(ComponentCRC, MaterialInterfaceCRC);
-
-						TArray<UTexture*> Textures;
-						MaterialInterface->GetUsedTextures(Textures, EMaterialQualityLevel::High, true, ERHIFeatureLevel::SM5, true);
-						for (UTexture* Texture : Textures)
-						{
-							uint32 TextureCRC = UHLODProxy::GetCRC(Texture);
-							UE_LOG(LogHLODBuilder, VeryVerbose, TEXT("     - Texture (%s) = %x"), *Texture->GetName(), TextureCRC);
-							ComponentCRC = HashCombine(ComponentCRC, TextureCRC);
-						}
-					}
-				}
-			}
-
-			ComponentsCRCs.Add(ComponentCRC);
-		}
-	}
-
-	// Sort the components CRCs to ensure the order of components won't have an impact on the final CRC
-	ComponentsCRCs.Sort();
+	// Min Visible Distance
+	uint32 HLODMinVisibleDistanceHash = GetTypeHash(InHLODActor->GetMinVisibleDistance());
+	UE_LOG(LogHLODBuilder, VeryVerbose, TEXT(" - HLOD Min Visible Distance (%.02f) = %x"), InHLODActor->GetMinVisibleDistance(), HLODMinVisibleDistanceHash);
+	Ar << HLODMinVisibleDistanceHash;
 
 	// Append all components CRCs
-	Ar << ComponentsCRCs;
+	uint32 HLODComponentsHash = UHLODBuilder::ComputeHLODHash(InActors);
+	UE_LOG(LogHLODBuilder, VeryVerbose, TEXT(" - HLOD Source Components = %x"), HLODComponentsHash);
+	Ar << HLODComponentsHash;
 
 	return Ar.GetCrc();
 }
@@ -290,6 +241,13 @@ TArray<AWorldPartitionHLOD*> FWorldPartitionHLODUtilities::CreateHLODActors(FHLO
 			bIsDirty = true;
 		}
 
+		// Minimum visible distance
+		if (!FMath::IsNearlyEqual(HLODActor->GetMinVisibleDistance(), InCreationParams.MinVisibleDistance))
+		{
+			HLODActor->SetMinVisibleDistance(InCreationParams.MinVisibleDistance);
+			bIsDirty = true;
+		}
+
 		// If any change was performed, mark HLOD package as dirty
 		if (bIsDirty)
 		{
@@ -335,18 +293,45 @@ TSubclassOf<UHLODBuilder> FWorldPartitionHLODUtilities::GetHLODBuilderClass(cons
 
 UHLODBuilderSettings* FWorldPartitionHLODUtilities::CreateHLODBuilderSettings(UHLODLayer* InHLODLayer)
 {
+	// Retrieve the HLOD builder class
 	TSubclassOf<UHLODBuilder> HLODBuilderClass = GetHLODBuilderClass(InHLODLayer);
 	if (!HLODBuilderClass)
 	{
 		return NewObject<UHLODBuilderSettings>(InHLODLayer, UHLODBuilderSettings::StaticClass());
 	}
 
-	UHLODBuilderSettings* HLODBuilderSettings = HLODBuilderClass->GetDefaultObject<UHLODBuilder>()->CreateSettings(InHLODLayer);
-	if (!ensure(HLODBuilderSettings))
+	// Retrieve the HLOD builder settings class
+	TSubclassOf<UHLODBuilderSettings> HLODBuilderSettingsClass = HLODBuilderClass->GetDefaultObject<UHLODBuilder>()->GetSettingsClass();
+	if (!ensure(HLODBuilderSettingsClass))
 	{
 		return NewObject<UHLODBuilderSettings>(InHLODLayer, UHLODBuilderSettings::StaticClass());
 	}
 
+	UHLODBuilderSettings* HLODBuilderSettings = NewObject<UHLODBuilderSettings>(InHLODLayer, HLODBuilderSettingsClass);
+
+	// Deprecated properties handling
+	if (InHLODLayer->GetHLODBuilderSettings() == nullptr)
+	{
+		EHLODLayerType HLODLayerType = InHLODLayer->GetLayerType();
+		switch (HLODLayerType)
+		{
+		case EHLODLayerType::MeshMerge:
+			CastChecked<UHLODBuilderMeshMergeSettings>(HLODBuilderSettings)->MeshMergeSettings = InHLODLayer->MeshMergeSettings_DEPRECATED;
+			CastChecked<UHLODBuilderMeshMergeSettings>(HLODBuilderSettings)->HLODMaterial = InHLODLayer->HLODMaterial_DEPRECATED;
+			break;
+
+		case EHLODLayerType::MeshSimplify:
+			CastChecked<UHLODBuilderMeshSimplifySettings>(HLODBuilderSettings)->MeshSimplifySettings = InHLODLayer->MeshSimplifySettings_DEPRECATED;
+			CastChecked<UHLODBuilderMeshSimplifySettings>(HLODBuilderSettings)->HLODMaterial = InHLODLayer->HLODMaterial_DEPRECATED;
+			break;
+
+		case EHLODLayerType::MeshApproximate:
+			CastChecked<UHLODBuilderMeshApproximateSettings>(HLODBuilderSettings)->MeshApproximationSettings = InHLODLayer->MeshApproximationSettings_DEPRECATED;
+			CastChecked<UHLODBuilderMeshApproximateSettings>(HLODBuilderSettings)->HLODMaterial = InHLODLayer->HLODMaterial_DEPRECATED;
+			break;
+		};
+	}
+
 	return HLODBuilderSettings;
 }
 
@@ -377,13 +362,73 @@ uint32 FWorldPartitionHLODUtilities::BuildHLOD(AWorldPartitionHLOD* InHLODActor)
 		if (ensure(HLODBuilder))
 		{
 			HLODBuilder->AddToRoot();
+
+			HLODBuilder->SetHLODBuilderSettings(HLODLayer->GetHLODBuilderSettings());
+
 			if (HLODBuilder->RequiresCompiledAssets())
 			{
 				// Wait for compilation to finish
 				FAssetCompilingManager::Get().FinishAllCompilation();
 			}
 
-			HLODBuilder->Build(InHLODActor, HLODLayer, LevelStreaming->GetLoadedLevel()->Actors);
+			FHLODBuildContext HLODBuildContext;
+			HLODBuildContext.World = InHLODActor->GetWorld();
+			HLODBuildContext.AssetsOuter = InHLODActor->GetPackage();
+			HLODBuildContext.AssetsBaseName = InHLODActor->GetActorLabel();
+			HLODBuildContext.MinVisibleDistance = InHLODActor->GetMinVisibleDistance();
+
+			TArray<UActorComponent*> HLODComponents = HLODBuilder->Build(HLODBuildContext, LevelStreaming->GetLoadedLevel()->Actors);
+			if (HLODComponents.IsEmpty())
+			{
+				UE_LOG(LogHLODBuilder, Warning, TEXT("HLOD generation created no component for %s"), *InHLODActor->GetActorLabel());
+			}
+
+			// Ideally, this should be performed elsewhere, to allow more flexibility in the HLOD generation
+			for (UActorComponent* HLODComponent : HLODComponents)
+			{
+				if (UPrimitiveComponent* HLODPrimitive = Cast<UPrimitiveComponent>(HLODComponent))
+				{
+					// Disable collisions
+					HLODPrimitive->SetCollisionProfileName(UCollisionProfile::NoCollision_ProfileName);
+					HLODPrimitive->SetGenerateOverlapEvents(false);
+					HLODPrimitive->SetCanEverAffectNavigation(false);
+					HLODPrimitive->CanCharacterStepUpOn = ECanBeCharacterBase::ECB_No;
+					HLODPrimitive->SetCanEverAffectNavigation(false);
+					HLODPrimitive->SetCollisionEnabled(ECollisionEnabled::NoCollision);
+					HLODPrimitive->SetMobility(EComponentMobility::Static);
+
+					// Enable optimizations
+					HLODPrimitive->bComputeFastLocalBounds = true;
+					HLODPrimitive->bComputeBoundsOnceForGame = true;
+				}
+
+				if (UStaticMeshComponent* StaticMeshComponent = Cast<UStaticMeshComponent>(HLODComponent))
+				{
+					if (UStaticMesh* StaticMesh = StaticMeshComponent->GetStaticMesh())
+					{
+						// If the HLOD process did create this static mesh
+						if (StaticMeshComponent->GetPackage() == StaticMesh->GetPackage())
+						{
+							// Set up ray tracing far fields for always loaded HLODs
+							StaticMeshComponent->bRayTracingFarField = !HLODLayer->IsSpatiallyLoaded() && StaticMesh->bSupportRayTracing;
+
+							// Disable collisions
+							if (UBodySetup* BodySetup = StaticMesh->GetBodySetup())
+							{
+								BodySetup->DefaultInstance.SetCollisionProfileName(UCollisionProfile::NoCollision_ProfileName);
+								BodySetup->CollisionTraceFlag = CTF_UseSimpleAsComplex;
+							}
+
+							// Rename owned static mesh
+							StaticMesh->Rename(*MakeUniqueObjectName(StaticMesh->GetOuter(), StaticMesh->GetClass(), *FString::Printf(TEXT("StaticMesh_%s"), *HLODLayer->GetName())).ToString());
+						}
+					}
+				}
+			}
+
+			InHLODActor->Modify();
+			InHLODActor->SetHLODComponents(HLODComponents);
+
 			HLODBuilder->RemoveFromRoot();
 		}
 			
diff --git a/Engine/Plugins/Editor/WorldPartitionHLODUtilities/Source/Public/HLODBuilderInstancing.h b/Engine/Plugins/Editor/WorldPartitionHLODUtilities/Source/Public/HLODBuilderInstancing.h
--- a/Engine/Plugins/Editor/WorldPartitionHLODUtilities/Source/Public/HLODBuilderInstancing.h
+++ b/Engine/Plugins/Editor/WorldPartitionHLODUtilities/Source/Public/HLODBuilderInstancing.h
@@ -18,5 +18,5 @@ public:
 	virtual bool RequiresCompiledAssets() const override { return false; }
 	virtual bool RequiresWarmup() const override { return false; }
 
-	virtual TArray<UPrimitiveComponent*> CreateComponents(AWorldPartitionHLOD* InHLODActor, const UHLODLayer* InHLODLayer, const TArray<UPrimitiveComponent*>& InSubComponents) const override;
+	virtual TArray<UActorComponent*> Build(const FHLODBuildContext& InHLODBuildContext, const TArray<UActorComponent*>& InSourceComponents) const override;
 };
diff --git a/Engine/Plugins/Editor/WorldPartitionHLODUtilities/Source/Public/HLODBuilderMeshApproximate.h b/Engine/Plugins/Editor/WorldPartitionHLODUtilities/Source/Public/HLODBuilderMeshApproximate.h
--- a/Engine/Plugins/Editor/WorldPartitionHLODUtilities/Source/Public/HLODBuilderMeshApproximate.h
+++ b/Engine/Plugins/Editor/WorldPartitionHLODUtilities/Source/Public/HLODBuilderMeshApproximate.h
@@ -36,6 +36,6 @@ class WORLDPARTITIONHLODUTILITIES_API UHLODBuilderMeshApproximate : public UHLOD
 	GENERATED_UCLASS_BODY()
 
 public:
-	virtual UHLODBuilderSettings* CreateSettings(UHLODLayer* InHLODLayer) const override;
-	virtual TArray<UPrimitiveComponent*> CreateComponents(AWorldPartitionHLOD* InHLODActor, const UHLODLayer* InHLODLayer, const TArray<UPrimitiveComponent*>& InSubComponents) const override;
+	virtual TSubclassOf<UHLODBuilderSettings> GetSettingsClass() const override;
+	virtual TArray<UActorComponent*> Build(const FHLODBuildContext& InHLODBuildContext, const TArray<UActorComponent*>& InSourceComponents) const override;
 };
diff --git a/Engine/Plugins/Editor/WorldPartitionHLODUtilities/Source/Public/HLODBuilderMeshMerge.h b/Engine/Plugins/Editor/WorldPartitionHLODUtilities/Source/Public/HLODBuilderMeshMerge.h
--- a/Engine/Plugins/Editor/WorldPartitionHLODUtilities/Source/Public/HLODBuilderMeshMerge.h
+++ b/Engine/Plugins/Editor/WorldPartitionHLODUtilities/Source/Public/HLODBuilderMeshMerge.h
@@ -36,6 +36,6 @@ class WORLDPARTITIONHLODUTILITIES_API UHLODBuilderMeshMerge : public UHLODBuilde
 	GENERATED_UCLASS_BODY()
 
 public:
-	virtual UHLODBuilderSettings* CreateSettings(UHLODLayer* InHLODLayer) const override;
-	virtual TArray<UPrimitiveComponent*> CreateComponents(AWorldPartitionHLOD* InHLODActor, const UHLODLayer* InHLODLayer, const TArray<UPrimitiveComponent*>& InSubComponents) const override;
+	virtual TSubclassOf<UHLODBuilderSettings> GetSettingsClass() const override;
+	virtual TArray<UActorComponent*> Build(const FHLODBuildContext& InHLODBuildContext, const TArray<UActorComponent*>& InSourceComponents) const override;
 };
diff --git a/Engine/Plugins/Editor/WorldPartitionHLODUtilities/Source/Public/HLODBuilderMeshSimplify.h b/Engine/Plugins/Editor/WorldPartitionHLODUtilities/Source/Public/HLODBuilderMeshSimplify.h
--- a/Engine/Plugins/Editor/WorldPartitionHLODUtilities/Source/Public/HLODBuilderMeshSimplify.h
+++ b/Engine/Plugins/Editor/WorldPartitionHLODUtilities/Source/Public/HLODBuilderMeshSimplify.h
@@ -36,6 +36,6 @@ class WORLDPARTITIONHLODUTILITIES_API UHLODBuilderMeshSimplify : public UHLODBui
 	GENERATED_UCLASS_BODY()
 
 public:
-	virtual UHLODBuilderSettings* CreateSettings(UHLODLayer* InHLODLayer) const override;
-	virtual TArray<UPrimitiveComponent*> CreateComponents(AWorldPartitionHLOD* InHLODActor, const UHLODLayer* InHLODLayer, const TArray<UPrimitiveComponent*>& InSubComponents) const override;
+	virtual TSubclassOf<UHLODBuilderSettings> GetSettingsClass() const override;
+	virtual TArray<UActorComponent*> Build(const FHLODBuildContext& InHLODBuildContext, const TArray<UActorComponent*>& InSourceComponents) const override;
 };
diff --git a/Engine/Source/Runtime/Engine/Classes/Components/ActorComponent.h b/Engine/Source/Runtime/Engine/Classes/Components/ActorComponent.h
--- a/Engine/Source/Runtime/Engine/Classes/Components/ActorComponent.h
+++ b/Engine/Source/Runtime/Engine/Classes/Components/ActorComponent.h
@@ -791,6 +791,15 @@ public:
 	 * data from GetEditorPreviewInfo().
 	 */
 	virtual TSharedPtr<SWidget> GetCustomEditorPreviewWidget() { return TSharedPtr<SWidget>(); }
+
+	/**
+	 * Return the custom HLODBuilder class that should be used to generate HLODs for components of this type.
+	 * Allows the HLOD system to include whatever is needed to represent a component at a distance.
+	 * For example, this is currently used to build custom HLODs for landscape streaming proxies. 
+	 * This could be used to generate fake lights, represent custom FX at a distance, or even to insert
+	 * externally generated HLODs.
+	 */
+	virtual TSubclassOf<class UHLODBuilder> GetCustomHLODBuilderClass() const { return nullptr; }
 #endif // WITH_EDITOR
 
 	/**
diff --git a/Engine/Source/Runtime/Engine/Private/WorldPartition/HLOD/HLODActor.cpp b/Engine/Source/Runtime/Engine/Private/WorldPartition/HLOD/HLODActor.cpp
--- a/Engine/Source/Runtime/Engine/Private/WorldPartition/HLOD/HLODActor.cpp
+++ b/Engine/Source/Runtime/Engine/Private/WorldPartition/HLOD/HLODActor.cpp
@@ -106,37 +106,36 @@ TUniquePtr<FWorldPartitionActorDesc> AWorldPartitionHLOD::CreateClassActorDesc()
 	return TUniquePtr<FWorldPartitionActorDesc>(new FHLODActorDesc());
 }
 
-void AWorldPartitionHLOD::SetHLODPrimitives(const TArray<UPrimitiveComponent*>& InHLODPrimitives)
+void AWorldPartitionHLOD::SetHLODComponents(const TArray<UActorComponent*>& InHLODComponents)
 {
-	TRACE_CPUPROFILER_EVENT_SCOPE(AWorldPartitionHLOD::SetHLODPrimitive);
-	check(!InHLODPrimitives.IsEmpty());
+	TRACE_CPUPROFILER_EVENT_SCOPE(AWorldPartitionHLOD::SetHLODComponents);
+	check(!InHLODComponents.IsEmpty());
 
-	TArray<USceneComponent*> ComponentsToRemove;
-	GetComponents<USceneComponent>(ComponentsToRemove);
-
-	SetRootComponent(InHLODPrimitives[0]);
-
-	for(UPrimitiveComponent* InHLODPrimitive : InHLODPrimitives)
+	TArray<UActorComponent*> ComponentsToRemove;
+	GetComponents(ComponentsToRemove);
+	for (UActorComponent* ComponentToRemove : ComponentsToRemove)
 	{
-		ComponentsToRemove.Remove(InHLODPrimitive);
+		ComponentToRemove->DestroyComponent();
+	}
 
-		AddInstanceComponent(InHLODPrimitive);
+	for(UActorComponent* HLODComponent : InHLODComponents)
+	{
+		HLODComponent->Rename(nullptr, this);
+		AddInstanceComponent(HLODComponent);
 
-		if (InHLODPrimitive != RootComponent)
+		if (USceneComponent* HLODSceneComponent = Cast<USceneComponent>(HLODComponent))
 		{
-			InHLODPrimitive->SetupAttachment(RootComponent);
+			if (RootComponent == nullptr)
+			{
+				SetRootComponent(HLODSceneComponent);
+			}
+			else
+			{
+				HLODSceneComponent->SetupAttachment(RootComponent);
+			}
 		}
 	
-		InHLODPrimitive->SetCollisionEnabled(ECollisionEnabled::NoCollision);
-		InHLODPrimitive->SetMobility(EComponentMobility::Static);
-
-		InHLODPrimitive->RegisterComponent();
-		InHLODPrimitive->MarkRenderStateDirty();
-	}
-
-	for (USceneComponent* ComponentToRemove : ComponentsToRemove)
-	{
-		ComponentToRemove->DestroyComponent();
+		HLODComponent->RegisterComponent();
 	}
 }
 
diff --git a/Engine/Source/Runtime/Engine/Private/WorldPartition/HLOD/HLODBuilder.cpp b/Engine/Source/Runtime/Engine/Private/WorldPartition/HLOD/HLODBuilder.cpp
--- a/Engine/Source/Runtime/Engine/Private/WorldPartition/HLOD/HLODBuilder.cpp
+++ b/Engine/Source/Runtime/Engine/Private/WorldPartition/HLOD/HLODBuilder.cpp
@@ -2,14 +2,7 @@
 
 #include "WorldPartition/HLOD/HLODBuilder.h"
 
-#include "Engine/StaticMesh.h"
-#include "Components/InstancedStaticMeshComponent.h"
-#include "PhysicsEngine/BodySetup.h"
-
-#include "WorldPartition/HLOD/HLODActor.h"
-#include "WorldPartition/HLOD/HLODLayer.h"
-
-#include "AssetCompilingManager.h"
+#include "Engine/HLODProxy.h"
 
 DEFINE_LOG_CATEGORY(LogHLODBuilder);
 
@@ -26,9 +19,15 @@ UHLODBuilderSettings::UHLODBuilderSettings(const FObjectInitializer& ObjectIniti
 
 #if WITH_EDITOR
 
-UHLODBuilderSettings* UHLODBuilder::CreateSettings(UHLODLayer* InHLODLayer) const
+TSubclassOf<UHLODBuilderSettings> UHLODBuilder::GetSettingsClass() const
+{
+	return UHLODBuilderSettings::StaticClass();
+}
+
+void UHLODBuilder::SetHLODBuilderSettings(const UHLODBuilderSettings* InHLODBuilderSettings)
 {
-	return NewObject<UHLODBuilderSettings>(InHLODLayer);
+	check(InHLODBuilderSettings->IsA(GetSettingsClass()));
+	HLODBuilderSettings = InHLODBuilderSettings;
 }
 
 bool UHLODBuilder::RequiresCompiledAssets() const
@@ -41,116 +40,99 @@ bool UHLODBuilder::RequiresWarmup() const
 	return true;
 }
 
-TArray<UPrimitiveComponent*> UHLODBuilder::CreateComponents(AWorldPartitionHLOD* InHLODActor, const UHLODLayer* InHLODLayer, const TArray<UPrimitiveComponent*>& InSubComponents) const
+static TArray<UActorComponent*> GatherHLODRelevantComponents(const TArray<AActor*>& InSourceActors)
 {
-	return TArray<UPrimitiveComponent*>();
-}
+	TArray<UActorComponent*> HLODRelevantComponents;
 
-void UHLODBuilder::Build(AWorldPartitionHLOD* InHLODActor, const UHLODLayer* InHLODLayer, const TArray<AActor*>& InSubActors)
-{
-	TArray<UPrimitiveComponent*> SubComponents = GatherPrimitiveComponents(InSubActors);
-	if (SubComponents.IsEmpty())
+	for (AActor* Actor : InSourceActors)
 	{
-		return;
-	}
-
-	TArray<UPrimitiveComponent*> HLODPrimitives = CreateComponents(InHLODActor, InHLODLayer, SubComponents);
-	HLODPrimitives.RemoveSwap(nullptr);
+		if (!Actor->IsHLODRelevant())
+		{
+			continue;
+		}
 
-	for (UPrimitiveComponent* HLODPrimitive : HLODPrimitives)
-	{
-		// Disable collisions
-		HLODPrimitive->SetCollisionProfileName(UCollisionProfile::NoCollision_ProfileName);
-		HLODPrimitive->SetGenerateOverlapEvents(false);
-		HLODPrimitive->SetCanEverAffectNavigation(false);
-		HLODPrimitive->CanCharacterStepUpOn = ECanBeCharacterBase::ECB_No;
-		HLODPrimitive->SetCanEverAffectNavigation(false);
-		HLODPrimitive->SetCollisionEnabled(ECollisionEnabled::NoCollision);
-
-		// Enable optimizations
-		HLODPrimitive->bComputeFastLocalBounds = true;
-		HLODPrimitive->bComputeBoundsOnceForGame = true;
-
-		if (InHLODLayer->GetLayerType() != EHLODLayerType::Instancing)
+		for (UActorComponent* SubComponent : Actor->GetComponents())
 		{
-			if (UStaticMeshComponent* StaticMeshComponent = Cast<UStaticMeshComponent>(HLODPrimitive))
+			if (SubComponent && SubComponent->IsHLODRelevant())
 			{
-				if (UStaticMesh* StaticMesh = StaticMeshComponent->GetStaticMesh())
-				{
-					// Set up ray tracing far fields
-					StaticMeshComponent->bRayTracingFarField = StaticMesh->bSupportRayTracing;
-
-					if (HLODPrimitive->GetPackage() == StaticMesh->GetPackage())
-					{
-						// Disable collisions on owned static mesh
-						if (UBodySetup* BodySetup = StaticMesh->GetBodySetup())
-						{							
-							BodySetup->DefaultInstance.SetCollisionProfileName(UCollisionProfile::NoCollision_ProfileName);
-							BodySetup->CollisionTraceFlag = CTF_UseSimpleAsComplex;
-						}
-						
-						// Rename owned static mesh
-						StaticMesh->Rename(*MakeUniqueObjectName(StaticMesh->GetOuter(), StaticMesh->GetClass(), *FString::Printf(TEXT("StaticMesh_%s"), *InHLODLayer->GetName())).ToString());
-					}
-				}
+				HLODRelevantComponents.Add(SubComponent);
 			}
 		}
 	}
 
-	if (!HLODPrimitives.IsEmpty())
+	return HLODRelevantComponents;
+}
+
+uint32 UHLODBuilder::ComputeHLODHash(const UActorComponent* InSourceComponent) const
+{
+	uint32 ComponentCRC = 0;
+
+	if (UStaticMeshComponent* StaticMeshComponent = Cast<UStaticMeshComponent>(const_cast<UActorComponent*>(InSourceComponent)))
 	{
-		InHLODActor->Modify();
-		InHLODActor->SetHLODPrimitives(HLODPrimitives);
+		UE_LOG(LogHLODBuilder, VeryVerbose, TEXT(" - Component \'%s\' from actor \'%s\'"), *StaticMeshComponent->GetName(), *StaticMeshComponent->GetOwner()->GetName());
+
+		ComponentCRC = UHLODProxy::GetCRC(StaticMeshComponent);
+		UE_LOG(LogHLODBuilder, VeryVerbose, TEXT("     - StaticMeshComponent (%s) = %x"), *StaticMeshComponent->GetName(), ComponentCRC);
 	}
+	else
+	{
+		ComponentCRC = FMath::Rand();
+		UE_LOG(LogHLODBuilder, Warning, TEXT("Can't compute HLOD hash for component of type %s, assuming it is dirty."), *InSourceComponent->GetClass()->GetName());
+		
+	}
+
+	return ComponentCRC;
 }
 
-TArray<UPrimitiveComponent*> UHLODBuilder::GatherPrimitiveComponents(const TArray<AActor*>& InActors)
+uint32 UHLODBuilder::ComputeHLODHash(const TArray<AActor*>& InSourceActors)
 {
-	TArray<UPrimitiveComponent*> PrimitiveComponents;
+	// We get the CRC of each component
+	TArray<uint32> ComponentsCRCs;
 
-	auto GatherPrimitivesFromActor = [&PrimitiveComponents](const AActor* Actor, const AActor* ParentActor = nullptr)
+	for (UActorComponent* SourceComponent : GatherHLODRelevantComponents(InSourceActors))
 	{
-		const TCHAR* Padding = ParentActor ? TEXT("    ") : TEXT("");
-		UE_LOG(LogHLODBuilder, Verbose, TEXT("%s* Adding components from actor %s"), Padding, *Actor->GetName());
-		for (UActorComponent* SubComponent : Actor->GetComponents())
+		TSubclassOf<UHLODBuilder> HLODBuilderClass = SourceComponent->GetCustomHLODBuilderClass();
+		if (!HLODBuilderClass)
 		{
-			if (SubComponent && SubComponent->IsHLODRelevant())
-			{
-				if (UPrimitiveComponent* PrimitiveComponent = Cast<UPrimitiveComponent>(SubComponent))
-				{
-					PrimitiveComponents.Add(PrimitiveComponent);
-
-					if (UInstancedStaticMeshComponent* ISMC = Cast<UInstancedStaticMeshComponent>(PrimitiveComponent))
-					{
-						UE_LOG(LogHLODBuilder, Verbose, TEXT("%s    * %s [%d instances]"), Padding, *ISMC->GetStaticMesh()->GetName(), ISMC->GetInstanceCount());
-					}
-					else if (UStaticMeshComponent* SMC = Cast<UStaticMeshComponent>(PrimitiveComponent))
-					{
-						UE_LOG(LogHLODBuilder, Verbose, TEXT("%s    * %s"), Padding, *SMC->GetStaticMesh()->GetName());
-					}
-				}
-				else
-				{
-					UE_LOG(LogHLODBuilder, Warning, TEXT("Component \"%s\" is marked as HLOD-relevant but this type of component currently unsupported."), *SubComponent->GetFullName());
-				}
-			}
+			HLODBuilderClass = UHLODBuilder::StaticClass();
 		}
-	};
 
-	TSet<AActor*> UnderlyingActors;
+		uint32 ComponentHash = HLODBuilderClass->GetDefaultObject<UHLODBuilder>()->ComputeHLODHash(SourceComponent);
+		ComponentsCRCs.Add(ComponentHash);
+	}
+
+	// Sort the components CRCs to ensure the order of components won't have an impact on the final CRC
+	ComponentsCRCs.Sort();
+
+	return FCrc::MemCrc32(ComponentsCRCs.GetData(), ComponentsCRCs.Num() * ComponentsCRCs.GetTypeSize());
+}
+
+TArray<UActorComponent*> UHLODBuilder::Build(const FHLODBuildContext& InHLODBuildContext, const TArray<AActor*>& InSourceActors) const
+{
+	TMap<TSubclassOf<UHLODBuilder>, TArray<UActorComponent*>> HLODBuildersForComponents;
 
-	for (AActor* Actor : InActors)
+	// Gather custom HLOD builders, and regroup all components by builders
+	for (UActorComponent* SourceComponent : GatherHLODRelevantComponents(InSourceActors))
 	{
-		if (!Actor->IsHLODRelevant())
-		{
-			continue;
-		}
+		TSubclassOf<UHLODBuilder> HLODBuilderClass = SourceComponent->GetCustomHLODBuilderClass();
+		HLODBuildersForComponents.FindOrAdd(HLODBuilderClass).Add(SourceComponent);
+	}
+	
+	TArray<UActorComponent*> HLODComponents;
+	for (const auto& HLODBuilderPair : HLODBuildersForComponents)
+	{
+		// If no custom HLOD builder is provided, use this current builder.
+		const UHLODBuilder* HLODBuilder = HLODBuilderPair.Key ? HLODBuilderPair.Key->GetDefaultObject<UHLODBuilder>() : this;
+		const TArray<UActorComponent*>& SourceComponents = HLODBuilderPair.Value;
 
-		// Gather primitives from the Actor
-		GatherPrimitivesFromActor(Actor);
+		TArray<UActorComponent*> NewComponents = HLODBuilder->Build(InHLODBuildContext, SourceComponents);
+		HLODComponents.Append(NewComponents);
 	}
 
-	return PrimitiveComponents;
+	// In case a builder returned null entries, clean the array.
+	HLODComponents.RemoveSwap(nullptr);
+
+	return HLODComponents;
 }
 
 #endif // WITH_EDITOR
diff --git a/Engine/Source/Runtime/Engine/Private/WorldPartition/RuntimeSpatialHash/RuntimeSpatialHashHLOD.cpp b/Engine/Source/Runtime/Engine/Private/WorldPartition/RuntimeSpatialHash/RuntimeSpatialHashHLOD.cpp
--- a/Engine/Source/Runtime/Engine/Private/WorldPartition/RuntimeSpatialHash/RuntimeSpatialHashHLOD.cpp
+++ b/Engine/Source/Runtime/Engine/Private/WorldPartition/RuntimeSpatialHash/RuntimeSpatialHashHLOD.cpp
@@ -187,6 +187,7 @@ static TArray<FGuid> GenerateHLODsForGrid(UWorldPartition* WorldPartition, const
 				CreationParams.CellName = CellName;
 				CreationParams.CellBounds = CellBounds;
 				CreationParams.HLODLevel = HLODLevel;
+				CreationParams.MinVisibleDistance = RuntimeGrid.LoadingRange;
 
 				IWorldPartitionHLODUtilities* WPHLODUtilities = FModuleManager::Get().LoadModuleChecked<IWorldPartitionHLODUtilitiesModule>("WorldPartitionHLODUtilities").GetUtilities();
 
diff --git a/Engine/Source/Runtime/Engine/Public/WorldPartition/HLOD/HLODActor.h b/Engine/Source/Runtime/Engine/Public/WorldPartition/HLOD/HLODActor.h
--- a/Engine/Source/Runtime/Engine/Public/WorldPartition/HLOD/HLODActor.h
+++ b/Engine/Source/Runtime/Engine/Public/WorldPartition/HLOD/HLODActor.h
@@ -28,7 +28,7 @@ public:
 	bool DoesRequireWarmup() const { return bRequireWarmup; }
 
 #if WITH_EDITOR
-	void SetHLODPrimitives(const TArray<UPrimitiveComponent*>& InHLODPrimitives);
+	void SetHLODComponents(const TArray<UActorComponent*>& InHLODComponents);
 
 	void SetSubActors(const TArray<FHLODSubActor>& InSubActorMappings);
 	const TArray<FHLODSubActor>& GetSubActors() const;
@@ -58,6 +58,9 @@ public:
 	const FBox& GetHLODBounds() const;
 	void SetHLODBounds(const FBox& InBounds);
 
+	double GetMinVisibleDistance() const { return MinVisibleDistance; }
+	void SetMinVisibleDistance(double InMinVisibleDistance) { MinVisibleDistance = InMinVisibleDistance; }
+
 	void BuildHLOD(bool bForceBuild = false);
 	uint32 GetHLODHash() const;
 #endif // WITH_EDITOR
@@ -109,6 +112,9 @@ private:
 	UPROPERTY()
 	FBox HLODBounds;
 
+	UPROPERTY()
+	double MinVisibleDistance;
+
 	UPROPERTY()
 	uint32 HLODHash;
 #endif
diff --git a/Engine/Source/Runtime/Engine/Public/WorldPartition/HLOD/HLODBuilder.h b/Engine/Source/Runtime/Engine/Public/WorldPartition/HLOD/HLODBuilder.h
--- a/Engine/Source/Runtime/Engine/Public/WorldPartition/HLOD/HLODBuilder.h
+++ b/Engine/Source/Runtime/Engine/Public/WorldPartition/HLOD/HLODBuilder.h
@@ -3,29 +3,56 @@
 #pragma once
 
 #include "CoreMinimal.h"
+#include "Algo/Transform.h"
+
 #include "HLODBuilder.generated.h"
 
 class AActor;
-class AWorldPartitionHLOD;
-class UPrimitiveComponent;
-class UHLODLayer;
+class UActorComponent;
 
 ENGINE_API DECLARE_LOG_CATEGORY_EXTERN(LogHLODBuilder, Log, All);
 
-
-// This is the base class for all HLOD Builder settings
+/**
+ * Base class for all HLOD Builder settings
+ */
 UCLASS()
 class ENGINE_API UHLODBuilderSettings : public UObject
 {
 	GENERATED_UCLASS_BODY()
+
 #if WITH_EDITOR
+public:
+	/**
+	 * Hash the class settings.
+	 * This is used to detect settings changes and trigger an HLOD rebuild if necessarry.
+	 */
 	virtual uint32 GetCRC() const { return 0; }
 #endif // WITH_EDITOR
 };
 
 
 /**
- * Base class for all HLODBuilders
+ * Provide context for the HLOD generation
+ */
+struct FHLODBuildContext
+{
+	/** World for which HLODs are being built */
+	UWorld*		World;
+
+	/** Outer to use for generated assets */
+	UObject*	AssetsOuter;
+
+	/** Base name to use for generated assets */
+	FString		AssetsBaseName;
+
+	/** Minimum distance at which the HLOD is expected to be displayed */
+	double		MinVisibleDistance;
+};
+
+
+/**
+ * Base class for all HLOD builders
+ * This class takes as input a group of components, and should return component(s) that will be included in the HLOD actor.
  */
 UCLASS(Abstract)
 class ENGINE_API UHLODBuilder : public UObject
@@ -34,16 +61,73 @@ class ENGINE_API UHLODBuilder : public UObject
 
 #if WITH_EDITOR
 public:
-	virtual UHLODBuilderSettings* CreateSettings(UHLODLayer* InHLODLayer) const;
+	/**
+	 * Provide builder settings before a Build.
+	 */
+	void SetHLODBuilderSettings(const UHLODBuilderSettings* InHLODBuilderSettings);
 
+	/**
+	 * Build an HLOD representation of the input actors.
+	 * Components returned by this method needs to be properly outered & assigned to your target (HLOD) actor.
+	 */
+	TArray<UActorComponent*> Build(const FHLODBuildContext& InHLODBuildContext, const TArray<AActor*>& InSourceActors) const;
+
+	/**
+	 * Return the setting subclass associated with this builder.
+	 */
+	virtual TSubclassOf<UHLODBuilderSettings> GetSettingsClass() const;
+
+	/**
+	 * Should return true if this builder needs to have it's input assets fully compiled.
+	 */
 	virtual bool RequiresCompiledAssets() const;
 
+	/**
+	 * Should return true if components generated from this builder need a warmup phase before being made visible.
+	 * If your components are using virtual textures or Nanite meshes, this should return true, as it will be necessary
+	 * to warmup the VT & Nanite caches before transitionning to HLOD. Otherwise, it's likely the initial first frames
+	 * could show a low resolution texture or mesh.
+	 */
 	virtual bool RequiresWarmup() const;
 
-	virtual TArray<UPrimitiveComponent*> CreateComponents(AWorldPartitionHLOD* InHLODActor, const UHLODLayer* InHLODLayer, const TArray<UPrimitiveComponent*>& InSubComponents) const;
+	/**
+	 * For a given component, compute a unique hash from the properties that are relevant to HLOD generation.
+	 * Used to detect changes to the source components of an HLOD.
+	 * The base version can only support hashing of static mesh components. HLOD builder subclasses
+	 * should override this method and compute the hash of component types they support as input.
+	 */
+	virtual uint32 ComputeHLODHash(const UActorComponent* InSourceComponent) const;
+
+	/**
+	 * Components created with this method needs to be properly outered & assigned to your target actor.
+	 */
+	virtual TArray<UActorComponent*> Build(const FHLODBuildContext& InHLODBuildContext, const TArray<UActorComponent*>& InSourceComponents) const PURE_VIRTUAL(UHLODBuilder::CreateComponents, return {};);
 
-	void Build(AWorldPartitionHLOD* InHLODActor, const UHLODLayer* InHLODLayer, const TArray<AActor*>& InSubActors);
+	/**
+	 * From a set of actors, compute a unique hash from their properties that are relevant to HLOD generation.
+	 * Used to detect changes to the source actors and trigger an HLOD rebuild if necessarry.
+	 */
+	static uint32 ComputeHLODHash(const TArray<AActor*>& InSourceActors);
+
+protected:
+	template <typename TComponentClass>
+	static inline TArray<TComponentClass*> FilterComponents(const TArray<UActorComponent*>& InSourceComponents)
+	{
+		TArray<TComponentClass*> FilteredComponents;
+		FilteredComponents.Reserve(InSourceComponents.Num());
+		Algo::TransformIf(InSourceComponents, FilteredComponents, [](UActorComponent* SourceComponent) { return Cast<TComponentClass>(SourceComponent); }, [](UActorComponent* SourceComponent) { return Cast<TComponentClass>(SourceComponent); });
+		if (InSourceComponents.Num() != FilteredComponents.Num())
+		{
+			UE_LOG(LogHLODBuilder, Warning, TEXT("Excluding %d components from the HLOD build."), InSourceComponents.Num() - FilteredComponents.Num());
+		}
+		return FilteredComponents;
+	}
 
-	static TArray<UPrimitiveComponent*> GatherPrimitiveComponents(const TArray<AActor*>& InActors);
 #endif // WITH_EDITOR
+
+#if WITH_EDITORONLY_DATA
+protected:
+	UPROPERTY()
+	TObjectPtr<const UHLODBuilderSettings> HLODBuilderSettings;
+#endif
 };
diff --git a/Engine/Source/Runtime/Engine/Public/WorldPartition/HLOD/IWorldPartitionHLODUtilities.h b/Engine/Source/Runtime/Engine/Public/WorldPartition/HLOD/IWorldPartitionHLODUtilities.h
--- a/Engine/Source/Runtime/Engine/Public/WorldPartition/HLOD/IWorldPartitionHLODUtilities.h
+++ b/Engine/Source/Runtime/Engine/Public/WorldPartition/HLOD/IWorldPartitionHLODUtilities.h
@@ -35,6 +35,8 @@ struct ENGINE_API FHLODCreationParams
 	FName CellName;
 	FBox  CellBounds;
 	uint32 HLODLevel;
+
+	double MinVisibleDistance;
 };
 
 /**
