Fix non-unity build

diff --git a/Engine/Plugins/Editor/WorldPartitionHLODUtilities/Source/Private/WorldPartitionHLODUtilities.cpp b/Engine/Plugins/Editor/WorldPartitionHLODUtilities/Source/Private/WorldPartitionHLODUtilities.cpp
--- a/Engine/Plugins/Editor/WorldPartitionHLODUtilities/Source/Private/WorldPartitionHLODUtilities.cpp
+++ b/Engine/Plugins/Editor/WorldPartitionHLODUtilities/Source/Private/WorldPartitionHLODUtilities.cpp
@@ -13,8 +13,6 @@
 #include "WorldPartition/HLOD/HLODLayer.h"
 #include "WorldPartition/WorldPartitionLevelStreamingDynamic.h"
 
-#include "ISMPartition/ISMComponentDescriptor.h"
-
 #include "Serialization/ArchiveCrc32.h"
 #include "Templates/UniquePtr.h"
 
@@ -27,6 +25,10 @@
 
 #include "Algo/Transform.h"
 
+#include "Engine/StaticMesh.h"
+#include "Components/StaticMeshComponent.h"
+#include "PhysicsEngine/BodySetup.h"
+
 static UWorldPartitionLevelStreamingDynamic* CreateLevelStreamingFromHLODActor(AWorldPartitionHLOD* InHLODActor, bool& bOutDirty)
 {
 	bOutDirty = false;
