Fixed failure to open/write an actor package due to async package save.

diff --git a/Engine/Plugins/Editor/WorldPartitionHLODUtilities/Source/Private/WorldPartitionHLODUtilities.cpp b/Engine/Plugins/Editor/WorldPartitionHLODUtilities/Source/Private/WorldPartitionHLODUtilities.cpp
--- a/Engine/Plugins/Editor/WorldPartitionHLODUtilities/Source/Private/WorldPartitionHLODUtilities.cpp
+++ b/Engine/Plugins/Editor/WorldPartitionHLODUtilities/Source/Private/WorldPartitionHLODUtilities.cpp
@@ -29,6 +29,8 @@
 
 static UWorldPartitionLevelStreamingDynamic* CreateLevelStreamingFromHLODActor(AWorldPartitionHLOD* InHLODActor, bool& bOutDirty)
 {
+	UPackage::WaitForAsyncFileWrites();
+
 	bOutDirty = false;
 	UWorld* World = InHLODActor->GetWorld();
 	UWorldPartition* WorldPartition = World->GetWorldPartition();
