Fixed non-spatially loaded HLODs components not being marked as bRayTracingFarField=true

diff --git a/Engine/Plugins/Editor/WorldPartitionHLODUtilities/Source/Private/WorldPartitionHLODUtilities.cpp b/Engine/Plugins/Editor/WorldPartitionHLODUtilities/Source/Private/WorldPartitionHLODUtilities.cpp
--- a/Engine/Plugins/Editor/WorldPartitionHLODUtilities/Source/Private/WorldPartitionHLODUtilities.cpp
+++ b/Engine/Plugins/Editor/WorldPartitionHLODUtilities/Source/Private/WorldPartitionHLODUtilities.cpp
@@ -13,6 +13,7 @@
 #include "WorldPartition/HLOD/HLODLayer.h"
 #include "WorldPartition/WorldPartitionLevelStreamingDynamic.h"
 
+#include "UObject/GCObjectScopeGuard.h"
 #include "Serialization/ArchiveCrc32.h"
 #include "Templates/UniquePtr.h"
 
@@ -363,7 +364,7 @@ uint32 FWorldPartitionHLODUtilities::BuildHLOD(AWorldPartitionHLOD* InHLODActor)
 		UHLODBuilder* HLODBuilder = NewObject<UHLODBuilder>(GetTransientPackage(), HLODBuilderClass);
 		if (ensure(HLODBuilder))
 		{
-			HLODBuilder->AddToRoot();
+			FGCObjectScopeGuard BuilderGCScopeGuard(HLODBuilder);
 
 			HLODBuilder->SetHLODBuilderSettings(HLODLayer->GetHLODBuilderSettings());
 
@@ -379,6 +380,9 @@ uint32 FWorldPartitionHLODUtilities::BuildHLOD(AWorldPartitionHLOD* InHLODActor)
 				UE_LOG(LogHLODBuilder, Warning, TEXT("HLOD generation created no component for %s"), *InHLODActor->GetActorLabel());
 			}
 
+			InHLODActor->Modify();
+			InHLODActor->SetHLODComponents(HLODComponents);
+
 			// Ideally, this should be performed elsewhere, to allow more flexibility in the HLOD generation
 			for (UActorComponent* HLODComponent : HLODComponents)
 			{
@@ -421,14 +425,7 @@ uint32 FWorldPartitionHLODUtilities::BuildHLOD(AWorldPartitionHLOD* InHLODActor)
 					}
 				}
 			}
-
-			InHLODActor->Modify();
-			InHLODActor->SetHLODComponents(HLODComponents);
-
-			HLODBuilder->RemoveFromRoot();
 		}
-			
-		InHLODActor->MarkPackageDirty();
 	}
 
 	return NewHLODHash;
