Modified WP HLOD settings to use UMaterialInterface rather than UMaterial

diff --git a/Engine/Plugins/Editor/WorldPartitionHLODUtilities/Source/Private/HLODBuilderMeshApproximate.cpp b/Engine/Plugins/Editor/WorldPartitionHLODUtilities/Source/Private/HLODBuilderMeshApproximate.cpp
--- a/Engine/Plugins/Editor/WorldPartitionHLODUtilities/Source/Private/HLODBuilderMeshApproximate.cpp
+++ b/Engine/Plugins/Editor/WorldPartitionHLODUtilities/Source/Private/HLODBuilderMeshApproximate.cpp
@@ -106,7 +106,7 @@ TArray<UActorComponent*> UHLODBuilderMeshApproximate::Build(const FHLODBuildCont
 
 	const UHLODBuilderMeshApproximateSettings* MeshApproximateSettings = CastChecked<UHLODBuilderMeshApproximateSettings>(HLODBuilderSettings);
 	const FMeshApproximationSettings& UseSettings = MeshApproximateSettings->MeshApproximationSettings;
-	UMaterial* HLODMaterial = MeshApproximateSettings->HLODMaterial.LoadSynchronous();
+	UMaterialInterface* HLODMaterial = MeshApproximateSettings->HLODMaterial.LoadSynchronous();
 
 	IGeometryProcessing_ApproximateActors::FOptions Options = ApproxActorsAPI->ConstructOptions(UseSettings);
 	Options.BasePackagePath = InHLODBuildContext.AssetsOuter->GetPackage()->GetName();
diff --git a/Engine/Plugins/Editor/WorldPartitionHLODUtilities/Source/Private/HLODBuilderMeshMerge.cpp b/Engine/Plugins/Editor/WorldPartitionHLODUtilities/Source/Private/HLODBuilderMeshMerge.cpp
--- a/Engine/Plugins/Editor/WorldPartitionHLODUtilities/Source/Private/HLODBuilderMeshMerge.cpp
+++ b/Engine/Plugins/Editor/WorldPartitionHLODUtilities/Source/Private/HLODBuilderMeshMerge.cpp
@@ -76,7 +76,7 @@ TArray<UActorComponent*> UHLODBuilderMeshMerge::Build(const FHLODBuildContext& I
 
 	const UHLODBuilderMeshMergeSettings* MeshMergeSettings = CastChecked<UHLODBuilderMeshMergeSettings>(HLODBuilderSettings);
 	const FMeshMergingSettings& UseSettings = MeshMergeSettings->MeshMergeSettings;
-	UMaterial* HLODMaterial = MeshMergeSettings->HLODMaterial.LoadSynchronous();
+	UMaterialInterface* HLODMaterial = MeshMergeSettings->HLODMaterial.LoadSynchronous();
 
 	const IMeshMergeUtilities& MeshMergeUtilities = FModuleManager::Get().LoadModuleChecked<IMeshMergeModule>("MeshMergeUtilities").GetUtilities();
 	MeshMergeUtilities.MergeComponentsToStaticMesh(SourcePrimitiveComponents, InHLODBuildContext.World, UseSettings, HLODMaterial, InHLODBuildContext.AssetsOuter->GetPackage(), InHLODBuildContext.AssetsBaseName, Assets, MergedActorLocation, 0.25f, false);
diff --git a/Engine/Plugins/Editor/WorldPartitionHLODUtilities/Source/Private/HLODBuilderMeshSimplify.cpp b/Engine/Plugins/Editor/WorldPartitionHLODUtilities/Source/Private/HLODBuilderMeshSimplify.cpp
--- a/Engine/Plugins/Editor/WorldPartitionHLODUtilities/Source/Private/HLODBuilderMeshSimplify.cpp
+++ b/Engine/Plugins/Editor/WorldPartitionHLODUtilities/Source/Private/HLODBuilderMeshSimplify.cpp
@@ -95,7 +95,7 @@ TArray<UActorComponent*> UHLODBuilderMeshSimplify::Build(const FHLODBuildContext
 
 	const UHLODBuilderMeshSimplifySettings* MeshSimplifySettings = CastChecked<UHLODBuilderMeshSimplifySettings>(HLODBuilderSettings);
 	const FMeshProxySettings& UseSettings = MeshSimplifySettings->MeshSimplifySettings;
-	UMaterial* HLODMaterial = MeshSimplifySettings->HLODMaterial.LoadSynchronous();
+	UMaterialInterface* HLODMaterial = MeshSimplifySettings->HLODMaterial.LoadSynchronous();
 
 	TArray<UObject*> Assets;
 	FCreateProxyDelegate ProxyDelegate;
diff --git a/Engine/Plugins/Editor/WorldPartitionHLODUtilities/Source/Public/HLODBuilderMeshApproximate.h b/Engine/Plugins/Editor/WorldPartitionHLODUtilities/Source/Public/HLODBuilderMeshApproximate.h
--- a/Engine/Plugins/Editor/WorldPartitionHLODUtilities/Source/Public/HLODBuilderMeshApproximate.h
+++ b/Engine/Plugins/Editor/WorldPartitionHLODUtilities/Source/Public/HLODBuilderMeshApproximate.h
@@ -23,7 +23,7 @@ class WORLDPARTITIONHLODUTILITIES_API UHLODBuilderMeshApproximateSettings : publ
 
 	/** Material that will be used by the generated HLOD static mesh */
 	UPROPERTY(EditAnywhere, Config, AdvancedDisplay, Category = HLOD)
-	TSoftObjectPtr<UMaterial> HLODMaterial;
+	TSoftObjectPtr<UMaterialInterface> HLODMaterial;
 };
 
 
diff --git a/Engine/Plugins/Editor/WorldPartitionHLODUtilities/Source/Public/HLODBuilderMeshMerge.h b/Engine/Plugins/Editor/WorldPartitionHLODUtilities/Source/Public/HLODBuilderMeshMerge.h
--- a/Engine/Plugins/Editor/WorldPartitionHLODUtilities/Source/Public/HLODBuilderMeshMerge.h
+++ b/Engine/Plugins/Editor/WorldPartitionHLODUtilities/Source/Public/HLODBuilderMeshMerge.h
@@ -23,7 +23,7 @@ class WORLDPARTITIONHLODUTILITIES_API UHLODBuilderMeshMergeSettings : public UHL
 
 	/** Material that will be used by the generated HLOD static mesh */
 	UPROPERTY(EditAnywhere, Config, AdvancedDisplay, Category = HLOD)
-	TSoftObjectPtr<UMaterial> HLODMaterial;
+	TSoftObjectPtr<UMaterialInterface> HLODMaterial;
 };
 
 
diff --git a/Engine/Plugins/Editor/WorldPartitionHLODUtilities/Source/Public/HLODBuilderMeshSimplify.h b/Engine/Plugins/Editor/WorldPartitionHLODUtilities/Source/Public/HLODBuilderMeshSimplify.h
--- a/Engine/Plugins/Editor/WorldPartitionHLODUtilities/Source/Public/HLODBuilderMeshSimplify.h
+++ b/Engine/Plugins/Editor/WorldPartitionHLODUtilities/Source/Public/HLODBuilderMeshSimplify.h
@@ -23,7 +23,7 @@ class WORLDPARTITIONHLODUTILITIES_API UHLODBuilderMeshSimplifySettings : public
 
 	/** Material that will be used by the generated HLOD static mesh */
 	UPROPERTY(EditAnywhere, Config, AdvancedDisplay, Category = HLOD)
-	TSoftObjectPtr<UMaterial> HLODMaterial;
+	TSoftObjectPtr<UMaterialInterface> HLODMaterial;
 };
 
 
