WP HLODs - Fixed issues with components attachement/mobility
* Ensure component mobility is setup properly before components are assigned to the HLOD actor. Would trigger warnings during HLOD generation ("AttachTo 'foo' is not static, cannot attach 'bar' which is static to it")
* Attachement - Make sure to keep components world transform intact when attaching to root component.
* Instance batching - Stationnary & static components would be batched separately

diff --git a/Engine/Plugins/Editor/WorldPartitionHLODUtilities/Source/Private/WorldPartitionHLODUtilities.cpp b/Engine/Plugins/Editor/WorldPartitionHLODUtilities/Source/Private/WorldPartitionHLODUtilities.cpp
--- a/Engine/Plugins/Editor/WorldPartitionHLODUtilities/Source/Private/WorldPartitionHLODUtilities.cpp
+++ b/Engine/Plugins/Editor/WorldPartitionHLODUtilities/Source/Private/WorldPartitionHLODUtilities.cpp
@@ -381,26 +381,28 @@ uint32 FWorldPartitionHLODUtilities::BuildHLOD(AWorldPartitionHLOD* InHLODActor)
 				UE_LOG(LogHLODBuilder, Warning, TEXT("HLOD generation created no component for %s"), *InHLODActor->GetActorLabel());
 			}
 
-			InHLODActor->Modify();
-			InHLODActor->SetHLODComponents(HLODComponents);
-
 			// Ideally, this should be performed elsewhere, to allow more flexibility in the HLOD generation
 			for (UActorComponent* HLODComponent : HLODComponents)
 			{
-				if (UPrimitiveComponent* HLODPrimitive = Cast<UPrimitiveComponent>(HLODComponent))
+				HLODComponent->SetCanEverAffectNavigation(false);
+
+				if (USceneComponent* SceneComponent = Cast<USceneComponent>(HLODComponent))
+				{
+					// Change Mobility to be Static
+					SceneComponent->SetMobility(EComponentMobility::Static);
+
+					// Enable bounds optimizations
+					SceneComponent->bComputeFastLocalBounds = true;
+					SceneComponent->bComputeBoundsOnceForGame = true;
+				}
+
+				if (UPrimitiveComponent* PrimitiveComponent = Cast<UPrimitiveComponent>(HLODComponent))
 				{
 					// Disable collisions
-					HLODPrimitive->SetCollisionProfileName(UCollisionProfile::NoCollision_ProfileName);
-					HLODPrimitive->SetGenerateOverlapEvents(false);
-					HLODPrimitive->SetCanEverAffectNavigation(false);
-					HLODPrimitive->CanCharacterStepUpOn = ECanBeCharacterBase::ECB_No;
-					HLODPrimitive->SetCanEverAffectNavigation(false);
-					HLODPrimitive->SetCollisionEnabled(ECollisionEnabled::NoCollision);
-					HLODPrimitive->SetMobility(EComponentMobility::Static);
-
-					// Enable optimizations
-					HLODPrimitive->bComputeFastLocalBounds = true;
-					HLODPrimitive->bComputeBoundsOnceForGame = true;
+					PrimitiveComponent->SetCollisionProfileName(UCollisionProfile::NoCollision_ProfileName);
+					PrimitiveComponent->SetGenerateOverlapEvents(false);
+					PrimitiveComponent->CanCharacterStepUpOn = ECanBeCharacterBase::ECB_No;
+					PrimitiveComponent->SetCollisionEnabled(ECollisionEnabled::NoCollision);
 				}
 
 				if (UStaticMeshComponent* StaticMeshComponent = Cast<UStaticMeshComponent>(HLODComponent))
@@ -408,10 +410,13 @@ uint32 FWorldPartitionHLODUtilities::BuildHLOD(AWorldPartitionHLOD* InHLODActor)
 					if (UStaticMesh* StaticMesh = StaticMeshComponent->GetStaticMesh())
 					{
 						// If the HLOD process did create this static mesh
-						if (StaticMeshComponent->GetPackage() == StaticMesh->GetPackage())
+						if (StaticMesh->GetPackage() == HLODBuildContext.AssetsOuter)
 						{
 							// Set up ray tracing far fields for always loaded HLODs
-							StaticMeshComponent->bRayTracingFarField = !HLODLayer->IsSpatiallyLoaded() && StaticMesh->bSupportRayTracing;
+							if (!HLODLayer->IsSpatiallyLoaded() && StaticMesh->bSupportRayTracing)
+							{
+								StaticMeshComponent->bRayTracingFarField = true;
+							}
 
 							// Disable collisions
 							if (UBodySetup* BodySetup = StaticMesh->GetBodySetup())
@@ -426,6 +431,8 @@ uint32 FWorldPartitionHLODUtilities::BuildHLOD(AWorldPartitionHLOD* InHLODActor)
 					}
 				}
 			}
+
+			InHLODActor->SetHLODComponents(HLODComponents);
 		}
 	}
 
diff --git a/Engine/Source/Runtime/Engine/Private/WorldPartition/HLOD/HLODActor.cpp b/Engine/Source/Runtime/Engine/Private/WorldPartition/HLOD/HLODActor.cpp
--- a/Engine/Source/Runtime/Engine/Private/WorldPartition/HLOD/HLODActor.cpp
+++ b/Engine/Source/Runtime/Engine/Private/WorldPartition/HLOD/HLODActor.cpp
@@ -41,11 +41,6 @@ AWorldPartitionHLOD::AWorldPartitionHLOD(const FObjectInitializer& ObjectInitial
 #endif
 }
 
-UPrimitiveComponent* AWorldPartitionHLOD::GetHLODComponent()
-{
-	return Cast<UPrimitiveComponent>(RootComponent);
-}
-
 void AWorldPartitionHLOD::SetVisibility(bool bInVisible)
 {
 	// When propagating visibility state to children, SetVisibility dirties all attached components.
@@ -110,6 +105,8 @@ void AWorldPartitionHLOD::SetHLODComponents(const TArray<UActorComponent*>& InHL
 {
 	TRACE_CPUPROFILER_EVENT_SCOPE(AWorldPartitionHLOD::SetHLODComponents);
 
+	Modify();
+
 	TArray<UActorComponent*> ComponentsToRemove;
 	GetComponents(ComponentsToRemove);
 	for (UActorComponent* ComponentToRemove : ComponentsToRemove)
@@ -130,7 +127,8 @@ void AWorldPartitionHLOD::SetHLODComponents(const TArray<UActorComponent*>& InHL
 			}
 			else
 			{
-				HLODSceneComponent->SetupAttachment(RootComponent);
+				// Attach to root component, but don't mess world tranform
+				HLODSceneComponent->AttachToComponent(RootComponent, FAttachmentTransformRules::KeepWorldTransform);
 			}
 		}
 	
diff --git a/Engine/Source/Runtime/Engine/Private/WorldPartition/HLOD/HLODBuilder.cpp b/Engine/Source/Runtime/Engine/Private/WorldPartition/HLOD/HLODBuilder.cpp
--- a/Engine/Source/Runtime/Engine/Private/WorldPartition/HLOD/HLODBuilder.cpp
+++ b/Engine/Source/Runtime/Engine/Private/WorldPartition/HLOD/HLODBuilder.cpp
@@ -149,6 +149,12 @@ namespace
 			// We'll always want to spawn ISMC, even if our source components are all SMC
 			ComponentClass = UInstancedStaticMeshComponent::StaticClass();
 
+			// Stationnary can be considered as static for the purpose of HLODs
+			if (Mobility == EComponentMobility::Stationary)
+			{
+				Mobility = EComponentMobility::Static;
+			}
+
 			ComputeHash();
 		}
 	};
diff --git a/Engine/Source/Runtime/Engine/Public/WorldPartition/HLOD/HLODActor.h b/Engine/Source/Runtime/Engine/Public/WorldPartition/HLOD/HLODActor.h
--- a/Engine/Source/Runtime/Engine/Public/WorldPartition/HLOD/HLODActor.h
+++ b/Engine/Source/Runtime/Engine/Public/WorldPartition/HLOD/HLODActor.h
@@ -90,8 +90,6 @@ protected:
 #endif
 	//~ End AActor Interface.
 
-	UPrimitiveComponent* GetHLODComponent();
-
 private:
 #if WITH_EDITORONLY_DATA
 	UPROPERTY()
