World Partition - HLOD: Do not assume that IMeshMergeUtilities::CreateProxyMesh() returns a single static mesh (another implementation could return multiple meshes)

diff --git a/Engine/Plugins/Editor/WorldPartitionHLODUtilities/Source/Private/HLODBuilderMeshSimplify.cpp b/Engine/Plugins/Editor/WorldPartitionHLODUtilities/Source/Private/HLODBuilderMeshSimplify.cpp
--- a/Engine/Plugins/Editor/WorldPartitionHLODUtilities/Source/Private/HLODBuilderMeshSimplify.cpp
+++ b/Engine/Plugins/Editor/WorldPartitionHLODUtilities/Source/Private/HLODBuilderMeshSimplify.cpp
@@ -83,20 +83,21 @@ TArray<UActorComponent*> UHLODBuilderMeshSimplify::Build(const FHLODBuildContext
 	const IMeshMergeUtilities& MeshMergeUtilities = FModuleManager::Get().LoadModuleChecked<IMeshMergeModule>("MeshMergeUtilities").GetUtilities();
 	MeshMergeUtilities.CreateProxyMesh(StaticMeshComponents, UseSettings, HLODMaterial, InHLODBuildContext.AssetsOuter->GetPackage(), InHLODBuildContext.AssetsBaseName, FGuid::NewGuid(), ProxyDelegate, true);
 
-	UStaticMeshComponent* Component = nullptr;
-	Algo::ForEach(Assets, [this, &Component](UObject* Asset)
+
+	TArray<UActorComponent*> Components;
+
+	Algo::ForEach(Assets, [this, &Components](UObject* Asset)
 	{
 		Asset->ClearFlags(RF_Public | RF_Standalone);
 
 		if (Cast<UStaticMesh>(Asset))
 		{
-			Component = NewObject<UStaticMeshComponent>();
-			Component->SetStaticMesh(static_cast<UStaticMesh*>(Asset));
+			UStaticMeshComponent* SMComponent = NewObject<UStaticMeshComponent>();
+			SMComponent->SetStaticMesh(static_cast<UStaticMesh*>(Asset));
+
+			Components.Add(SMComponent);
 		}
 	});
-
-	TArray<UActorComponent*> Components;
-	Components.Add(Component);
-
+	
 	return Components;
 }
