WorldPartition: Add ContentBundleGuid to Cell Packages (instead of comparing ContainerPackage and WorldPackage which doesn't translate to a ContentBundle all the time)

diff --git a/Engine/Plugins/Editor/WorldPartitionHLODUtilities/Source/Private/WorldPartitionHLODUtilities.cpp b/Engine/Plugins/Editor/WorldPartitionHLODUtilities/Source/Private/WorldPartitionHLODUtilities.cpp
--- a/Engine/Plugins/Editor/WorldPartitionHLODUtilities/Source/Private/WorldPartitionHLODUtilities.cpp
+++ b/Engine/Plugins/Editor/WorldPartitionHLODUtilities/Source/Private/WorldPartitionHLODUtilities.cpp
@@ -42,7 +42,8 @@ static UWorldPartitionLevelStreamingDynamic* CreateLevelStreamingFromHLODActor(A
 	const FName LevelStreamingName = FName(*FString::Printf(TEXT("HLODLevelStreaming_%s"), *InHLODActor->GetName()));
 	TArray<FWorldPartitionRuntimeCellObjectMapping> Mappings;
 	Mappings.Reserve(InHLODActor->GetSubActors().Num());
-	Algo::Transform(InHLODActor->GetSubActors(), Mappings, [World](const FHLODSubActor& SubActor) { return FWorldPartitionRuntimeCellObjectMapping(SubActor.ActorPackage, SubActor.ActorPath, SubActor.ContainerID, SubActor.ContainerTransform, SubActor.ContainerPackage, World->GetPackage()->GetFName()); });
+	// @todo_ow: investigate if we need to pass in a content bundle id here
+	Algo::Transform(InHLODActor->GetSubActors(), Mappings, [World](const FHLODSubActor& SubActor) { return FWorldPartitionRuntimeCellObjectMapping(SubActor.ActorPackage, SubActor.ActorPath, SubActor.ContainerID, SubActor.ContainerTransform, SubActor.ContainerPackage, World->GetPackage()->GetFName(), FGuid()); });
 
 	UWorldPartitionLevelStreamingDynamic* LevelStreaming = UWorldPartitionLevelStreamingDynamic::LoadInEditor(World, LevelStreamingName, Mappings);
 	check(LevelStreaming);
diff --git a/Engine/Source/Runtime/Engine/Private/WorldPartition/WorldPartitionLevelStreamingDynamic.cpp b/Engine/Source/Runtime/Engine/Private/WorldPartition/WorldPartitionLevelStreamingDynamic.cpp
--- a/Engine/Source/Runtime/Engine/Private/WorldPartition/WorldPartitionLevelStreamingDynamic.cpp
+++ b/Engine/Source/Runtime/Engine/Private/WorldPartition/WorldPartitionLevelStreamingDynamic.cpp
@@ -294,8 +294,9 @@ bool UWorldPartitionLevelStreamingDynamic::IssueLoadRequests()
 
 		for (const FWorldPartitionRuntimeCellObjectMapping& CellObjectMapping : ChildPackages)
 		{
-			if (CellObjectMapping.ContainerPackage != CellObjectMapping.WorldPackage)
+			if (CellObjectMapping.ContentBundleGuid.IsValid())
 			{
+				check(CellObjectMapping.ContainerPackage != CellObjectMapping.WorldPackage);
 				bool bIsContainerPackageAlreadyRemapped = OutLinkInstancingContext.RemapPackage(CellObjectMapping.ContainerPackage) != CellObjectMapping.ContainerPackage;
 				if (!bIsContainerPackageAlreadyRemapped)
 				{
diff --git a/Engine/Source/Runtime/Engine/Private/WorldPartition/WorldPartitionRuntimeLevelStreamingCell.cpp b/Engine/Source/Runtime/Engine/Private/WorldPartition/WorldPartitionRuntimeLevelStreamingCell.cpp
--- a/Engine/Source/Runtime/Engine/Private/WorldPartition/WorldPartitionRuntimeLevelStreamingCell.cpp
+++ b/Engine/Source/Runtime/Engine/Private/WorldPartition/WorldPartitionRuntimeLevelStreamingCell.cpp
@@ -103,7 +103,7 @@ void UWorldPartitionRuntimeLevelStreamingCell::AddActorToCell(const FWorldPartit
 	check(!ActorDescView.GetActorIsEditorOnly());
 	// Leaving this using deprecated functions until the serialization format of Packages is updated.
 PRAGMA_DISABLE_DEPRECATION_WARNINGS
-	Packages.Emplace(ActorDescView.GetActorPackage(), ActorDescView.GetActorPath(), InContainerID, InContainerTransform, InContainer->GetContainerPackage(), GetWorld()->GetPackage()->GetFName());
+	Packages.Emplace(ActorDescView.GetActorPackage(), ActorDescView.GetActorPath(), InContainerID, InContainerTransform, InContainer->GetContainerPackage(), GetWorld()->GetPackage()->GetFName(), ActorDescView.GetContentBundleUID());
 PRAGMA_ENABLE_DEPRECATION_WARNINGS
 }
 
diff --git a/Engine/Source/Runtime/Engine/Public/WorldPartition/WorldPartitionRuntimeCell.h b/Engine/Source/Runtime/Engine/Public/WorldPartition/WorldPartitionRuntimeCell.h
--- a/Engine/Source/Runtime/Engine/Public/WorldPartition/WorldPartitionRuntimeCell.h
+++ b/Engine/Source/Runtime/Engine/Public/WorldPartition/WorldPartitionRuntimeCell.h
@@ -38,7 +38,7 @@ struct FWorldPartitionRuntimeCellObjectMapping
 #endif
 	{}
 
-	FWorldPartitionRuntimeCellObjectMapping(FName InPackage, FName InPath, const FActorContainerID& InContainerID, const FTransform& InContainerTransform, FName InContainerPackage, FName InWorldPackage)
+	FWorldPartitionRuntimeCellObjectMapping(FName InPackage, FName InPath, const FActorContainerID& InContainerID, const FTransform& InContainerTransform, FName InContainerPackage, FName InWorldPackage, const FGuid& InContentBundleGuid)
 #if WITH_EDITORONLY_DATA
 		: Package(InPackage)
 		, Path(InPath)
@@ -46,6 +46,7 @@ struct FWorldPartitionRuntimeCellObjectMapping
 		, ContainerTransform(InContainerTransform)
 		, ContainerPackage(InContainerPackage)
 		, WorldPackage(InWorldPackage)
+		, ContentBundleGuid(InContentBundleGuid)
 		, LoadedPath(InPath)
 #endif
 	{}
@@ -87,6 +88,12 @@ struct FWorldPartitionRuntimeCellObjectMapping
 	UPROPERTY()
 	FName WorldPackage;
 
+	/**
+	 * Content Bundle Id
+	 */
+	UPROPERTY()
+	FGuid ContentBundleGuid;
+
 	/**
 	* Loaded actor path (when cooking or pie)
 	* 
