World Partition - HLOD: Improved DisplayName for "HLODMaterial" -> "HLOD Material"

diff --git a/Engine/Plugins/Editor/WorldPartitionHLODUtilities/Source/Public/HLODBuilderMeshApproximate.h b/Engine/Plugins/Editor/WorldPartitionHLODUtilities/Source/Public/HLODBuilderMeshApproximate.h
--- a/Engine/Plugins/Editor/WorldPartitionHLODUtilities/Source/Public/HLODBuilderMeshApproximate.h
+++ b/Engine/Plugins/Editor/WorldPartitionHLODUtilities/Source/Public/HLODBuilderMeshApproximate.h
@@ -22,7 +22,7 @@ class WORLDPARTITIONHLODUTILITIES_API UHLODBuilderMeshApproximateSettings : publ
 	FMeshApproximationSettings MeshApproximationSettings;
 
 	/** Material that will be used by the generated HLOD static mesh */
-	UPROPERTY(EditAnywhere, Config, AdvancedDisplay, Category = HLOD)
+	UPROPERTY(EditAnywhere, Config, AdvancedDisplay, Category = HLOD, meta = (DisplayName = "HLOD Material"))
 	TSoftObjectPtr<UMaterialInterface> HLODMaterial;
 };
 
diff --git a/Engine/Plugins/Editor/WorldPartitionHLODUtilities/Source/Public/HLODBuilderMeshMerge.h b/Engine/Plugins/Editor/WorldPartitionHLODUtilities/Source/Public/HLODBuilderMeshMerge.h
--- a/Engine/Plugins/Editor/WorldPartitionHLODUtilities/Source/Public/HLODBuilderMeshMerge.h
+++ b/Engine/Plugins/Editor/WorldPartitionHLODUtilities/Source/Public/HLODBuilderMeshMerge.h
@@ -22,7 +22,7 @@ class WORLDPARTITIONHLODUTILITIES_API UHLODBuilderMeshMergeSettings : public UHL
 	FMeshMergingSettings MeshMergeSettings;
 
 	/** Material that will be used by the generated HLOD static mesh */
-	UPROPERTY(EditAnywhere, Config, AdvancedDisplay, Category = HLOD)
+	UPROPERTY(EditAnywhere, Config, AdvancedDisplay, Category = HLOD, meta = (DisplayName = "HLOD Material"))
 	TSoftObjectPtr<UMaterialInterface> HLODMaterial;
 };
 
diff --git a/Engine/Plugins/Editor/WorldPartitionHLODUtilities/Source/Public/HLODBuilderMeshSimplify.h b/Engine/Plugins/Editor/WorldPartitionHLODUtilities/Source/Public/HLODBuilderMeshSimplify.h
--- a/Engine/Plugins/Editor/WorldPartitionHLODUtilities/Source/Public/HLODBuilderMeshSimplify.h
+++ b/Engine/Plugins/Editor/WorldPartitionHLODUtilities/Source/Public/HLODBuilderMeshSimplify.h
@@ -22,7 +22,7 @@ class WORLDPARTITIONHLODUTILITIES_API UHLODBuilderMeshSimplifySettings : public
 	FMeshProxySettings MeshSimplifySettings;
 
 	/** Material that will be used by the generated HLOD static mesh */
-	UPROPERTY(EditAnywhere, Config, AdvancedDisplay, Category = HLOD)
+	UPROPERTY(EditAnywhere, Config, AdvancedDisplay, Category = HLOD, meta = (DisplayName = "HLOD Material"))
 	TSoftObjectPtr<UMaterialInterface> HLODMaterial;
 };
 
