World Partition - HLOD: Instancing now has a "Disallow Nanite" option
* Necessary on Nanite meshes that have multiple LODs, otherwise the setting the mesh last LOD as the component ForcedLOD will not work

diff --git a/Engine/Plugins/Editor/WorldPartitionHLODUtilities/Source/Private/HLODBuilderInstancing.cpp b/Engine/Plugins/Editor/WorldPartitionHLODUtilities/Source/Private/HLODBuilderInstancing.cpp
--- a/Engine/Plugins/Editor/WorldPartitionHLODUtilities/Source/Private/HLODBuilderInstancing.cpp
+++ b/Engine/Plugins/Editor/WorldPartitionHLODUtilities/Source/Private/HLODBuilderInstancing.cpp
@@ -1,17 +1,67 @@
 // Copyright Epic Games, Inc. All Rights Reserved.
 
 #include "HLODBuilderInstancing.h"
+#include "Serialization/ArchiveCrc32.h"
+#include "Engine/StaticMesh.h"
+#include "Components/StaticMeshComponent.h"
 
 #include UE_INLINE_GENERATED_CPP_BY_NAME(HLODBuilderInstancing)
 
 
+UHLODBuilderInstancingSettings::UHLODBuilderInstancingSettings(const FObjectInitializer& ObjectInitializer)
+	: Super(ObjectInitializer)
+	, bDisallowNanite(false)
+{
+}
+
+uint32 UHLODBuilderInstancingSettings::GetCRC() const
+{
+	UHLODBuilderInstancingSettings& This = *const_cast<UHLODBuilderInstancingSettings*>(this);
+
+	FArchiveCrc32 Ar;
+
+	Ar << This.bDisallowNanite;
+	UE_LOG(LogHLODBuilder, VeryVerbose, TEXT(" - bDisallowNanite = %d"), Ar.GetCrc());
+
+	return Ar.GetCrc();
+}
+
+
 UHLODBuilderInstancing::UHLODBuilderInstancing(const FObjectInitializer& ObjectInitializer)
 	: Super(ObjectInitializer)
 {
 }
 
+TSubclassOf<UHLODBuilderSettings> UHLODBuilderInstancing::GetSettingsClass() const
+{
+	return UHLODBuilderInstancingSettings::StaticClass();
+}
+
 TArray<UActorComponent*> UHLODBuilderInstancing::Build(const FHLODBuildContext& InHLODBuildContext, const TArray<UActorComponent*>& InSourceComponents) const
 {
-	return UHLODBuilder::BatchInstances(InSourceComponents);
+	const UHLODBuilderInstancingSettings& InstancingSettings = *CastChecked<UHLODBuilderInstancingSettings>(HLODBuilderSettings);
+
+	TArray<UActorComponent*> HLODComponents = UHLODBuilder::BatchInstances(InSourceComponents);
+
+	// If requested, disallow Nanite on components whose mesh is Nanite enabled and has multiple LODs
+	if (InstancingSettings.bDisallowNanite)
+	{
+		for (UActorComponent* HLODComponent : HLODComponents)
+		{
+			if (UStaticMeshComponent* SMComponent = Cast<UStaticMeshComponent>(HLODComponent))
+			{
+				if (UStaticMesh* StaticMesh = SMComponent->GetStaticMesh())
+				{
+					if (!SMComponent->bDisallowNanite && StaticMesh->HasValidNaniteData() && StaticMesh->GetNumLODs() > 1)
+					{
+						SMComponent->bDisallowNanite = true;
+						SMComponent->MarkRenderStateDirty();
+					}
+				}
+			}
+		}
+	}
+
+	return HLODComponents;
 }
 
diff --git a/Engine/Plugins/Editor/WorldPartitionHLODUtilities/Source/Public/HLODBuilderInstancing.h b/Engine/Plugins/Editor/WorldPartitionHLODUtilities/Source/Public/HLODBuilderInstancing.h
--- a/Engine/Plugins/Editor/WorldPartitionHLODUtilities/Source/Public/HLODBuilderInstancing.h
+++ b/Engine/Plugins/Editor/WorldPartitionHLODUtilities/Source/Public/HLODBuilderInstancing.h
@@ -6,6 +6,22 @@
 #include "HLODBuilderInstancing.generated.h"
 
 
+UCLASS(Blueprintable, Config = Engine, PerObjectConfig)
+class WORLDPARTITIONHLODUTILITIES_API UHLODBuilderInstancingSettings : public UHLODBuilderSettings
+{
+	GENERATED_UCLASS_BODY()
+
+	virtual uint32 GetCRC() const override;
+
+	/**
+	 * If enabled, the components created for the HLODs will not use Nanite.
+	 * Necessary if you want to use the last LOD & the mesh is Nanite enabled, as forced LODs are ignored by Nanite
+	 */
+	UPROPERTY(EditAnywhere, AdvancedDisplay, BlueprintReadWrite, Category = Rendering)
+	bool bDisallowNanite;
+};
+
+
 /**
  * Build a AWorldPartitionHLOD whose components are ISMC
  */
@@ -18,5 +34,6 @@ public:
 	virtual bool RequiresCompiledAssets() const override { return false; }
 	virtual bool RequiresWarmup() const override { return false; }
 
+	virtual TSubclassOf<UHLODBuilderSettings> GetSettingsClass() const override;
 	virtual TArray<UActorComponent*> Build(const FHLODBuildContext& InHLODBuildContext, const TArray<UActorComponent*>& InSourceComponents) const override;
 };
