World Partition - HLOD: Changed HLOD actors folders & labels naming
* HLOD actors will be found in the outliner under the "HLOD/Name_Of_HLOD_Layer_Asset"
* HLOD actors label will be "Name_Of_HLOD_Layer_Asset/Cell_Name"

diff --git a/Engine/Plugins/Editor/WorldPartitionHLODUtilities/Source/Private/WorldPartitionHLODUtilities.cpp b/Engine/Plugins/Editor/WorldPartitionHLODUtilities/Source/Private/WorldPartitionHLODUtilities.cpp
--- a/Engine/Plugins/Editor/WorldPartitionHLODUtilities/Source/Private/WorldPartitionHLODUtilities.cpp
+++ b/Engine/Plugins/Editor/WorldPartitionHLODUtilities/Source/Private/WorldPartitionHLODUtilities.cpp
@@ -158,8 +158,6 @@ PRAGMA_ENABLE_DEPRECATION_WARNINGS
 			SpawnParams.NameMode = FActorSpawnParameters::ESpawnActorNameMode::Required_Fatal;
 			HLODActor = InCreationParams.WorldPartition->GetWorld()->SpawnActor<AWorldPartitionHLOD>(SpawnParams);
 
-			HLODActor->SetActorLabel(FString::Printf(TEXT("HLOD%d_%s"), InCreationParams.HLODLevel, *InCreationParams.CellName.ToString()));
-			HLODActor->SetFolderPath(*FString::Printf(TEXT("HLOD/HLOD%d"), InCreationParams.HLODLevel));
 			HLODActor->SetSourceCellName(InCreationParams.CellName);
 			HLODActor->SetSubActorsHLODLayer(HLODLayer);
 			HLODActor->SetGridIndices(InCreationParams.GridIndexX, InCreationParams.GridIndexY, InCreationParams.GridIndexZ);
@@ -243,6 +241,22 @@ PRAGMA_ENABLE_DEPRECATION_WARNINGS
 			bIsDirty = true;
 		}
 
+		// Actor label
+		const FString ActorLabel = FString::Printf(TEXT("%s/%s"), *HLODLayer->GetName(), *InCreationParams.CellName.ToString());
+		if (HLODActor->GetActorLabel() != ActorLabel)
+		{
+			HLODActor->SetActorLabel(ActorLabel);
+			bIsDirty = true;
+		}
+
+		// Folder name
+		const FName FolderPath(FString::Printf(TEXT("HLOD/%s"), *HLODLayer->GetName()));
+		if (HLODActor->GetFolderPath() != FolderPath)
+		{
+			HLODActor->SetFolderPath(FolderPath);
+			bIsDirty = true;
+		}
+
 		// Cell bounds
 		if (!HLODActor->GetHLODBounds().Equals(InCreationParams.CellBounds))
 		{
