WorldPartition - HLOD: Added support for HLOD stats stored as metadata in the HLOD Actor Descriptor
* Added HLOD Actor stats methods
* Added wp.Editor.HLOD.DumpStats console command to write a CSV containing the stats to /Saved/WorldPartition/
* Removed RuntimeSpatialHash related code in HLOD actor/actordesc/builder (x,y,z grid indices)

This will serve as the basis for showing the stats in-editor through the HLOD Outliner

diff --git a/Engine/Plugins/Editor/WorldPartitionHLODUtilities/Source/Private/WorldPartitionHLODUtilities.cpp b/Engine/Plugins/Editor/WorldPartitionHLODUtilities/Source/Private/WorldPartitionHLODUtilities.cpp
--- a/Engine/Plugins/Editor/WorldPartitionHLODUtilities/Source/Private/WorldPartitionHLODUtilities.cpp
+++ b/Engine/Plugins/Editor/WorldPartitionHLODUtilities/Source/Private/WorldPartitionHLODUtilities.cpp
@@ -15,20 +15,25 @@
 #include "WorldPartition/HLOD/HLODLayer.h"
 #include "WorldPartition/WorldPartitionLevelStreamingDynamic.h"
 
-#include "UObject/GCObjectScopeGuard.h"
-#include "Serialization/ArchiveCrc32.h"
-#include "Templates/UniquePtr.h"
-
 #include "HLODBuilderInstancing.h"
 #include "HLODBuilderMeshMerge.h"
 #include "HLODBuilderMeshSimplify.h"
 #include "HLODBuilderMeshApproximate.h"
 
 #include "Algo/Transform.h"
-
-#include "Engine/StaticMesh.h"
 #include "Components/StaticMeshComponent.h"
+#include "Engine/InstancedStaticMesh.h"
+#include "Engine/StaticMesh.h"
+#include "HAL/FileManager.h"
+#include "Materials/MaterialInstance.h"
 #include "PhysicsEngine/BodySetup.h"
+#include "ProfilingDebugging/ScopedTimers.h"
+#include "Serialization/ArchiveCrc32.h"
+#include "StaticMeshCompiler.h"
+#include "Templates/UniquePtr.h"
+#include "TextureCompiler.h"
+#include "UObject/GCObjectScopeGuard.h"
+
 
 static UWorldPartitionLevelStreamingDynamic* CreateLevelStreamingFromHLODActor(AWorldPartitionHLOD* InHLODActor, bool& bOutDirty)
 {
@@ -139,12 +144,19 @@ PRAGMA_ENABLE_DEPRECATION_WARNINGS
 		const FSubActorsInfo& SubActorsInfo = Pair.Value;
 		check(!SubActorsInfo.SubActors.IsEmpty());
 
-		// Compute HLODActor hash
-		uint64 CellHash = FHLODActorDesc::ComputeCellHash(HLODLayer->GetName(), InCreationParams.GridIndexX, InCreationParams.GridIndexY, InCreationParams.GridIndexZ, InCreationParams.DataLayersID);
+		auto ComputeCellHash = [](const FString& HLODLayerName, const FName CellName)
+		{
+			uint64 HLODLayerNameHash = FCrc::StrCrc32(*HLODLayerName);
+			uint32 CellNameHash = FCrc::StrCrc32(*CellName.ToString());
+			return HashCombine(HLODLayerNameHash, CellNameHash);
+		};
+
+		uint64 CellHash = ComputeCellHash(HLODLayer->GetName(), InCreationParams.CellName);
+		FName HLODActorName = *FString::Printf(TEXT("%s_%016llx"), *HLODLayer->GetName(), CellHash);		
 
 		AWorldPartitionHLOD* HLODActor = nullptr;
 		FWorldPartitionHandle HLODActorHandle;
-		if (InCreationContext.HLODActorDescs.RemoveAndCopyValue(CellHash, HLODActorHandle))
+		if (InCreationContext.HLODActorDescs.RemoveAndCopyValue(HLODActorName, HLODActorHandle))
 		{
 			InCreationContext.ActorReferences.Add(HLODActorHandle.ToReference());
 			HLODActor = CastChecked<AWorldPartitionHLOD>(HLODActorHandle->GetActor());
@@ -154,13 +166,12 @@ PRAGMA_ENABLE_DEPRECATION_WARNINGS
 		if (bNewActor)
 		{
 			FActorSpawnParameters SpawnParams;
-			SpawnParams.Name = *FString::Printf(TEXT("%s_%016llx"), *HLODLayer->GetName(), CellHash);
+			SpawnParams.Name = HLODActorName;
 			SpawnParams.NameMode = FActorSpawnParameters::ESpawnActorNameMode::Required_Fatal;
 			HLODActor = InCreationParams.WorldPartition->GetWorld()->SpawnActor<AWorldPartitionHLOD>(SpawnParams);
 
 			HLODActor->SetSourceCellName(InCreationParams.CellName);
 			HLODActor->SetSubActorsHLODLayer(HLODLayer);
-			HLODActor->SetGridIndices(InCreationParams.GridIndexX, InCreationParams.GridIndexY, InCreationParams.GridIndexZ);
 
 			// Make sure the generated HLOD actor has the same data layers as the source actors
 			for (const UDataLayerInstance* DataLayerInstance : InDataLayersInstances)
@@ -171,13 +182,8 @@ PRAGMA_ENABLE_DEPRECATION_WARNINGS
 		else
 		{
 #if DO_CHECK
-			uint64 GridIndexX, GridIndexY, GridIndexZ;
-			HLODActor->GetGridIndices(GridIndexX, GridIndexY, GridIndexZ);
-			check(GridIndexX == InCreationParams.GridIndexX);
-			check(GridIndexY == InCreationParams.GridIndexY);
-			check(GridIndexZ == InCreationParams.GridIndexZ);
+			check(HLODActor->GetSourceCellName() == InCreationParams.CellName);
 			check(HLODActor->GetSubActorsHLODLayer() == HLODLayer);
-			check(FDataLayersID(HLODActor->GetDataLayerInstances()) == InCreationParams.DataLayersID);
 #endif
 		}
 
@@ -358,10 +364,202 @@ UHLODBuilderSettings* FWorldPartitionHLODUtilities::CreateHLODBuilderSettings(UH
 	return HLODBuilderSettings;
 }
 
+void GatherInputStats(AWorldPartitionHLOD* InHLODActor, const TArray<AActor*> SubActors)
+{
+	int64 NumActors = 0;
+	int64 NumTriangles = 0;
+	int64 NumVertices = 0;
+
+	for (AActor* SubActor : SubActors)
+	{
+		if (SubActor && SubActor->IsHLODRelevant())
+		{
+			if (AWorldPartitionHLOD* SubHLODActor = Cast<AWorldPartitionHLOD>(SubActor))
+			{
+				NumActors += SubHLODActor->GetStat(FWorldPartitionHLODStats::InputActorCount);
+				NumTriangles += SubHLODActor->GetStat(FWorldPartitionHLODStats::InputTriangleCount);
+				NumVertices += SubHLODActor->GetStat(FWorldPartitionHLODStats::InputVertexCount);				
+			}
+			else
+			{
+				NumActors++;
+
+				SubActor->ForEachComponent<UStaticMeshComponent>(false, [&] (UStaticMeshComponent* SMComponent)
+				{
+					const UStaticMesh* StaticMesh = SMComponent->GetStaticMesh();
+					const FStaticMeshRenderData* RenderData = StaticMesh ? StaticMesh->GetRenderData() : nullptr;
+					const bool bHasRenderData = RenderData && !RenderData->LODResources.IsEmpty();
+
+					if (bHasRenderData)
+					{
+						const UInstancedStaticMeshComponent* ISMComponent = Cast<UInstancedStaticMeshComponent>(SMComponent);
+						const int64 LOD0TriCount = RenderData->LODResources[0].GetNumTriangles();
+						const int64 LOD0VtxCount = RenderData->LODResources[0].GetNumVertices();
+						const int64 NumInstances = ISMComponent ? ISMComponent->GetInstanceCount() : 1;
+
+						NumTriangles += LOD0TriCount * NumInstances;
+						NumVertices += LOD0VtxCount * NumInstances;
+					}
+				});
+			}
+		}
+	}
+
+	InHLODActor->SetStat(FWorldPartitionHLODStats::InputActorCount, NumActors);
+	InHLODActor->SetStat(FWorldPartitionHLODStats::InputTriangleCount, NumTriangles);
+	InHLODActor->SetStat(FWorldPartitionHLODStats::InputVertexCount, NumVertices);	
+}
+
+void GatherOutputStats(AWorldPartitionHLOD* InHLODActor)
+{
+	const UPackage* HLODPackage = InHLODActor->GetPackage();
+
+	// Gather relevant assets and process them outside of this ForEach as it's possible that async compilation completion
+	// triggers insertion of new objects, which would break the iteration over UObjects
+	TArray<UStaticMesh*> StaticMeshes;
+	TArray<UTexture*> Textures;
+	TArray<UMaterialInstance*> MaterialInstances;
+	ForEachObjectWithPackage(HLODPackage, [&](UObject* Object)
+	{
+		if (UStaticMesh* StaticMesh = Cast<UStaticMesh>(Object))
+		{
+			StaticMeshes.Add(StaticMesh);
+		}
+		else if (UMaterialInstance* MaterialInstance = Cast<UMaterialInstance>(Object))
+		{
+			MaterialInstances.Add(MaterialInstance);
+		}
+		else if (UTexture* Texture = Cast<UTexture>(Object))
+		{
+			Textures.Add(Texture);
+		}
+		return true;
+	}, false);
+
+	// Process static meshes
+	int64 MeshResourceSize = 0;
+	{
+		int64 InstanceCount = 0;
+		int64 NaniteTriangleCount = 0;
+		int64 NaniteVertexCount = 0;
+		int64 TriangleCount = 0;
+		int64 VertexCount = 0;
+		int64 UVChannelCount = 0;
+
+		InHLODActor->ForEachComponent<UInstancedStaticMeshComponent>(false, [&](const UInstancedStaticMeshComponent* ISMC)
+		{
+			InstanceCount += ISMC->GetInstanceCount();
+		});
+
+		FStaticMeshCompilingManager::Get().FinishCompilation(StaticMeshes);
+
+		for (UStaticMesh* StaticMesh : StaticMeshes)
+		{
+			MeshResourceSize += StaticMesh->GetResourceSizeBytes(EResourceSizeMode::Exclusive);
+
+			TriangleCount += StaticMesh->GetNumTriangles(0);
+			VertexCount += StaticMesh->GetNumVertices(0);
+			UVChannelCount = FMath::Max(UVChannelCount, StaticMesh->GetNumTexCoords(0));
+
+			NaniteTriangleCount += StaticMesh->GetNumNaniteTriangles();
+			NaniteVertexCount += StaticMesh->GetNumNaniteVertices();
+		}
+
+		// Mesh stats
+		InHLODActor->SetStat(FWorldPartitionHLODStats::MeshInstanceCount, InstanceCount);
+		InHLODActor->SetStat(FWorldPartitionHLODStats::MeshNaniteTriangleCount, NaniteTriangleCount);
+		InHLODActor->SetStat(FWorldPartitionHLODStats::MeshNaniteVertexCount, NaniteVertexCount);
+		InHLODActor->SetStat(FWorldPartitionHLODStats::MeshTriangleCount, TriangleCount);
+		InHLODActor->SetStat(FWorldPartitionHLODStats::MeshVertexCount, VertexCount);
+		InHLODActor->SetStat(FWorldPartitionHLODStats::MeshUVChannelCount, UVChannelCount);
+	}
+
+	// Process materials
+	int64 TexturesResourceSize = 0;
+	{
+		int64 BaseColorTextureSize = 0;
+		int64 NormalTextureSize = 0;
+		int64 EmissiveTextureSize = 0;
+		int64 MetallicTextureSize = 0;
+		int64 RoughnessTextureSize = 0;
+		int64 SpecularTextureSize = 0;
+
+		FTextureCompilingManager::Get().FinishCompilation(Textures);
+
+		for (UMaterialInstance* MaterialInstance : MaterialInstances)
+		{
+			// Retrieve the texture size for a texture that can have different names
+			auto GetTextureSize = [&](const TArray<FName>& TextureParamNames) -> int64
+			{
+				for (FName TextureParamName : TextureParamNames)
+				{
+					UTexture* Texture = nullptr;
+					MaterialInstance->GetTextureParameterValue(TextureParamName, Texture, true);
+
+					if (Texture && Texture->GetPackage() == HLODPackage)
+					{
+						TexturesResourceSize += Texture->GetResourceSizeBytes(EResourceSizeMode::Exclusive);
+						return Texture->GetSurfaceWidth();
+					}
+				}
+
+				return 0;
+			};
+
+			int64 LocalBaseColorTextureSize = GetTextureSize({ "BaseColorTexture", "DiffuseTexture" });
+			int64 LocalNormalTextureSize = GetTextureSize({ "NormalTexture" });
+			int64 LocalEmissiveTextureSize = GetTextureSize({ "EmissiveTexture", "EmissiveColorTexture" });
+			int64 LocalMetallicTextureSize = GetTextureSize({ "MetallicTexture" });
+			int64 LocalRoughnessTextureSize = GetTextureSize({ "RoughnessTexture" });
+			int64 LocalSpecularTextureSize = GetTextureSize({ "SpecularTexture" });
+
+			int64 MRSTextureSize = GetTextureSize({ "PackedTexture" });
+			if (MRSTextureSize != 0)
+			{
+				LocalMetallicTextureSize = LocalRoughnessTextureSize = LocalSpecularTextureSize = MRSTextureSize;
+			}
+
+			BaseColorTextureSize = FMath::Max(BaseColorTextureSize, LocalBaseColorTextureSize);
+			NormalTextureSize = FMath::Max(NormalTextureSize, LocalNormalTextureSize);
+			EmissiveTextureSize = FMath::Max(EmissiveTextureSize, LocalEmissiveTextureSize);
+			MetallicTextureSize = FMath::Max(MetallicTextureSize, LocalMetallicTextureSize);
+			RoughnessTextureSize = FMath::Max(RoughnessTextureSize, LocalRoughnessTextureSize);
+			SpecularTextureSize = FMath::Max(SpecularTextureSize, LocalSpecularTextureSize);
+		}
+
+		// Material stats
+		InHLODActor->SetStat(FWorldPartitionHLODStats::MaterialBaseColorTextureSize, BaseColorTextureSize);
+		InHLODActor->SetStat(FWorldPartitionHLODStats::MaterialNormalTextureSize, NormalTextureSize);
+		InHLODActor->SetStat(FWorldPartitionHLODStats::MaterialEmissiveTextureSize, EmissiveTextureSize);
+		InHLODActor->SetStat(FWorldPartitionHLODStats::MaterialMetallicTextureSize, MetallicTextureSize);
+		InHLODActor->SetStat(FWorldPartitionHLODStats::MaterialRoughnessTextureSize, RoughnessTextureSize);
+		InHLODActor->SetStat(FWorldPartitionHLODStats::MaterialSpecularTextureSize, SpecularTextureSize);
+	}
+
+	// Memory stats
+	InHLODActor->SetStat(FWorldPartitionHLODStats::MemoryMeshResourceSizeBytes, MeshResourceSize);
+	InHLODActor->SetStat(FWorldPartitionHLODStats::MemoryTexturesResourceSizeBytes, TexturesResourceSize);
+	InHLODActor->SetStat(FWorldPartitionHLODStats::MemoryDiskSizeBytes, 0); // Clear disk size as it is unknown at this point. It will be assigned at package loading
+}
+
 uint32 FWorldPartitionHLODUtilities::BuildHLOD(AWorldPartitionHLOD* InHLODActor)
 {
+	FAutoScopedDurationTimer TotalTimeScope;
+
+	// Keep track of timings related to this build
+	int64 LoadTimeMS = 0;
+	int64 BuildTimeMS = 0;
+	int64 TotalTimeMS = 0;
+	
+	// Load actors relevant to HLODs
 	bool bIsDirty = false;
-	UWorldPartitionLevelStreamingDynamic* LevelStreaming = CreateLevelStreamingFromHLODActor(InHLODActor, bIsDirty);
+	UWorldPartitionLevelStreamingDynamic* LevelStreaming = nullptr;
+	{
+		FAutoScopedDurationTimer LoadTimeScope;
+		LevelStreaming = CreateLevelStreamingFromHLODActor(InHLODActor, bIsDirty);
+		LoadTimeMS = LoadTimeScope.GetTime() * 1000;
+	}
+
 	ON_SCOPE_EXIT
 	{
 		UWorldPartitionLevelStreamingDynamic::UnloadFromEditor(LevelStreaming);
@@ -376,6 +574,12 @@ uint32 FWorldPartitionHLODUtilities::BuildHLOD(AWorldPartitionHLOD* InHLODActor)
 		return OldHLODHash;
 	}
 
+	// Clear stats as we're about to refresh them
+	InHLODActor->ResetStats();
+	
+	// Gather stats from the input to our HLOD build
+	GatherInputStats(InHLODActor, LevelStreaming->GetLoadedLevel()->Actors);
+	
 	const UHLODLayer* HLODLayer = InHLODActor->GetSubActorsHLODLayer();
 	TSubclassOf<UHLODBuilder> HLODBuilderClass = GetHLODBuilderClass(HLODLayer);
 
@@ -394,7 +598,14 @@ uint32 FWorldPartitionHLODUtilities::BuildHLOD(AWorldPartitionHLOD* InHLODActor)
 			HLODBuildContext.AssetsBaseName = InHLODActor->GetActorLabel();
 			HLODBuildContext.MinVisibleDistance = InHLODActor->GetMinVisibleDistance();
 
-			TArray<UActorComponent*> HLODComponents = HLODBuilder->Build(HLODBuildContext, LevelStreaming->GetLoadedLevel()->Actors);
+			// Build
+			TArray<UActorComponent*> HLODComponents;
+			{
+				FAutoScopedDurationTimer BuildTimeScope;
+				HLODComponents = HLODBuilder->Build(HLODBuildContext, LevelStreaming->GetLoadedLevel()->Actors);
+				BuildTimeMS = BuildTimeScope.GetTime() * 1000;
+			}
+
 			if (HLODComponents.IsEmpty())
 			{
 				UE_LOG(LogHLODBuilder, Warning, TEXT("HLOD generation created no component for %s"), *InHLODActor->GetActorLabel());
@@ -438,6 +649,7 @@ uint32 FWorldPartitionHLODUtilities::BuildHLOD(AWorldPartitionHLOD* InHLODActor)
 							}
 
 							// Disable collisions
+							StaticMesh->MarkAsNotHavingNavigationData();
 							if (UBodySetup* BodySetup = StaticMesh->GetBodySetup())
 							{
 								BodySetup->DefaultInstance.SetCollisionProfileName(UCollisionProfile::NoCollision_ProfileName);
@@ -455,6 +667,16 @@ uint32 FWorldPartitionHLODUtilities::BuildHLOD(AWorldPartitionHLOD* InHLODActor)
 		}
 	}
 
+	// Gather stats pertaining to the assets generated during this build
+	GatherOutputStats(InHLODActor);
+
+	TotalTimeMS = TotalTimeScope.GetTime() * 1000;
+
+	// Build timings stats
+	InHLODActor->SetStat(FWorldPartitionHLODStats::BuildTimeLoadMilliseconds, LoadTimeMS);
+	InHLODActor->SetStat(FWorldPartitionHLODStats::BuildTimeBuildMilliseconds, BuildTimeMS);
+	InHLODActor->SetStat(FWorldPartitionHLODStats::BuildTimeTotalMilliseconds, TotalTimeMS);
+
 	return NewHLODHash;
 }
 
diff --git a/Engine/Source/Editor/UnrealEd/Private/WorldPartition/WorldPartitionHLODsBuilder.cpp b/Engine/Source/Editor/UnrealEd/Private/WorldPartition/WorldPartitionHLODsBuilder.cpp
--- a/Engine/Source/Editor/UnrealEd/Private/WorldPartition/WorldPartitionHLODsBuilder.cpp
+++ b/Engine/Source/Editor/UnrealEd/Private/WorldPartition/WorldPartitionHLODsBuilder.cpp
@@ -159,31 +159,38 @@ UWorldPartitionHLODsBuilder::UWorldPartitionHLODsBuilder(const FObjectInitialize
 	: Super(ObjectInitializer)
 	, BuilderIdx(INDEX_NONE)
 	, BuilderCount(INDEX_NONE)
-	, HLODLevelToBuild(INDEX_NONE)
 	, DistributedBuildWorkingDir(FPaths::RootDir() / DistributedBuildWorkingDirName)
 	, DistributedBuildManifest(DistributedBuildWorkingDir / DistributedBuildManifestName)
 {
 	BuildOptions = FParse::Param(FCommandLine::Get(), TEXT("SetupHLODs")) ? EHLODBuildStep::HLOD_Setup : EHLODBuildStep::None;
 	BuildOptions |= FParse::Param(FCommandLine::Get(), TEXT("BuildHLODs")) ? EHLODBuildStep::HLOD_Build : EHLODBuildStep::None;
+	BuildOptions |= FParse::Param(FCommandLine::Get(), TEXT("RebuildHLODs")) ? EHLODBuildStep::HLOD_Build : EHLODBuildStep::None;
 	BuildOptions |= FParse::Param(FCommandLine::Get(), TEXT("DeleteHLODs")) ? EHLODBuildStep::HLOD_Delete : EHLODBuildStep::None;
 	BuildOptions |= FParse::Param(FCommandLine::Get(), TEXT("FinalizeHLODs")) ? EHLODBuildStep::HLOD_Finalize : EHLODBuildStep::None;
 	BuildOptions |= FParse::Param(FCommandLine::Get(), TEXT("DumpStats")) ? EHLODBuildStep::HLOD_Stats : EHLODBuildStep::None;
 
-	// Default behavior without any option is to setup + build
-	if (BuildOptions == EHLODBuildStep::None)
-	{
-		BuildOptions = EHLODBuildStep::HLOD_Setup | EHLODBuildStep::HLOD_Build;
-	}
-
 	bResumeBuild = FParse::Value(FCommandLine::Get(), TEXT("ResumeBuild="), ResumeBuildIndex);
 
 	bDistributedBuild = FParse::Param(FCommandLine::Get(), TEXT("DistributedBuild"));
+	bForceBuild = FParse::Param(FCommandLine::Get(), TEXT("RebuildHLODs"));
 
 	FParse::Value(FCommandLine::Get(), TEXT("BuildManifest="), BuildManifest);
 	FParse::Value(FCommandLine::Get(), TEXT("BuilderIdx="), BuilderIdx);
 	FParse::Value(FCommandLine::Get(), TEXT("BuilderCount="), BuilderCount);
+	FParse::Value(FCommandLine::Get(), TEXT("BuildHLODLayer="), HLODLayerToBuild);
+	FParse::Value(FCommandLine::Get(), TEXT("BuildSingleHLOD="), HLODActorToBuild);
 
-	FParse::Value(FCommandLine::Get(), TEXT("HLODLevel="), HLODLevelToBuild);
+	if (!HLODActorToBuild.IsNone() || !HLODLayerToBuild.IsNone())
+	{
+		BuildOptions = EHLODBuildStep::HLOD_Build;
+		bForceBuild = bForceBuild || !HLODActorToBuild.IsNone();
+	}
+	
+	// Default behavior without any option is to setup + build
+	if (BuildOptions == EHLODBuildStep::None)
+	{
+		BuildOptions = EHLODBuildStep::HLOD_Setup | EHLODBuildStep::HLOD_Build;
+	}
 
 	if (bDistributedBuild)
 	{
@@ -329,9 +336,14 @@ bool UWorldPartitionHLODsBuilder::SetupHLODActors()
 {
 	auto ActorFolderAddedDelegateHandle = GEngine->OnActorFolderAdded().AddLambda([this](UActorFolder* InActorFolder)
 	{
-		// We don't want the HLOD folders to be expanded by default
-		InActorFolder->SetIsInitiallyExpanded(false);
-		SourceControlHelper->Save(InActorFolder->GetPackage());
+		UPackage* ActorFolderPackage = InActorFolder->GetPackage();
+		const bool bIsTempPackage = FPackageName::IsTempPackage(ActorFolderPackage->GetName());
+		if (!bIsTempPackage)
+		{
+			// We don't want the HLOD folders to be expanded by default
+			InActorFolder->SetIsInitiallyExpanded(false);
+			SourceControlHelper->Save(InActorFolder->GetPackage());
+		}
 	});
 	
 	ON_SCOPE_EXIT
@@ -467,13 +479,10 @@ bool UWorldPartitionHLODsBuilder::BuildHLODActors()
 
 		UE_LOG(LogWorldPartitionHLODsBuilder, Display, TEXT("[%d / %d] Building HLOD actor %s..."), CurrentActor + 1, HLODActorsToBuild.Num(), *HLODActor->GetActorLabel());
 
-		if (HLODLevelToBuild == INDEX_NONE || HLODActor->GetLODLevel() == HLODLevelToBuild)
-		{
-			// Simulate an engine tick to make sure engine & render resources that are queued for deletion are processed.
-			FWorldPartitionHelpers::FakeEngineTick(WorldPartition->GetWorld());
+		// Simulate an engine tick to make sure engine & render resources that are queued for deletion are processed.
+		FWorldPartitionHelpers::FakeEngineTick(WorldPartition->GetWorld());
 
-			HLODActor->BuildHLOD();
-		}
+		HLODActor->BuildHLOD(bForceBuild);
 
 		UPackage* ActorPackage = HLODActor->GetPackage();
 		if (ActorPackage->IsDirty())
@@ -643,7 +652,6 @@ bool UWorldPartitionHLODsBuilder::DumpStats()
 		HLODStat.Name = HLODActor->GetActorLabel();
 		HLODStat.RuntimeGrid = HLODActor->GetRuntimeGrid();
 		HLODStat.bIsSpatiallyLoaded = HLODActor->GetIsSpatiallyLoaded();
-		HLODActor->GetGridIndices(HLODStat.GridLocationX, HLODStat.GridLocationY, HLODStat.GridLocationZ);
 		HLODStat.DataLayers = *FString::JoinBy(HLODActor->GetDataLayerInstances(), TEXT("| "), [](const UDataLayerInstance* DataLayer) { return DataLayer->GetDataLayerShortName(); });
 		HLODStat.HLODType = HLODActor->GetSubActorsHLODLayer()->GetLayerType();
 
@@ -778,6 +786,18 @@ TArray<TArray<FGuid>> UWorldPartitionHLODsBuilder::GetHLODWorkloads(int32 NumWor
 	TMap<FGuid, TArray<FGuid>>	HLODParenting;
 	for (FActorDescList::TIterator<AWorldPartitionHLOD> HLODIterator(WorldPartition->GetActorDescContainer()); HLODIterator; ++HLODIterator)
 	{
+		// Filter by HLOD actor
+		if (!HLODActorToBuild.IsNone() && HLODIterator->GetActorLabel() != HLODActorToBuild)
+		{
+			continue;
+		}
+
+		// Filter by HLOD layer
+		if (!HLODLayerToBuild.IsNone() && HLODIterator->GetSourceHLODLayerName() != HLODLayerToBuild)
+		{
+			continue;
+		}
+
 		TArray<FGuid>& ChildHLODs = HLODParenting.Add(HLODIterator->GetGuid());
 
 		for (const FHLODSubActorDesc& SubActor : HLODIterator->GetSubActors())
diff --git a/Engine/Source/Editor/UnrealEd/Public/WorldPartition/WorldPartitionHLODsBuilder.h b/Engine/Source/Editor/UnrealEd/Public/WorldPartition/WorldPartitionHLODsBuilder.h
--- a/Engine/Source/Editor/UnrealEd/Public/WorldPartition/WorldPartitionHLODsBuilder.h
+++ b/Engine/Source/Editor/UnrealEd/Public/WorldPartition/WorldPartitionHLODsBuilder.h
@@ -111,13 +111,14 @@ private:
 	EHLODBuildStep BuildOptions;
 
 	bool bDistributedBuild;
+	bool bForceBuild;
 	FString BuildManifest;
 	int32 BuilderIdx;
 	int32 BuilderCount;
 	bool bResumeBuild;
 	int32 ResumeBuildIndex;
-	int32 HLODLevelToBuild;
-
+	FName HLODLayerToBuild;
+	FName HLODActorToBuild;
 
 	const FString DistributedBuildWorkingDir;
 	const FString DistributedBuildManifest;
diff --git a/Engine/Source/Runtime/Core/Public/UObject/FortniteMainBranchObjectVersion.h b/Engine/Source/Runtime/Core/Public/UObject/FortniteMainBranchObjectVersion.h
--- a/Engine/Source/Runtime/Core/Public/UObject/FortniteMainBranchObjectVersion.h
+++ b/Engine/Source/Runtime/Core/Public/UObject/FortniteMainBranchObjectVersion.h
@@ -253,6 +253,9 @@ struct CORE_API FFortniteMainBranchObjectVersion
 		// Add Nanite Material Override option to materials and material instances.
 		NaniteMaterialOverride,
 
+		// Serialize HLOD stats in HLODActorDesc
+		WorldPartitionHLODActorDescSerializeStats,
+
 		// -----<new versions can be added above this line>-------------------------------------------------
 		VersionPlusOne,
 		LatestVersion = VersionPlusOne - 1
diff --git a/Engine/Source/Runtime/Engine/Private/WorldPartition/HLOD/HLODActor.cpp b/Engine/Source/Runtime/Engine/Private/WorldPartition/HLOD/HLODActor.cpp
--- a/Engine/Source/Runtime/Engine/Private/WorldPartition/HLOD/HLODActor.cpp
+++ b/Engine/Source/Runtime/Engine/Private/WorldPartition/HLOD/HLODActor.cpp
@@ -119,6 +119,9 @@ void AWorldPartitionHLOD::PostLoad()
 			}
 		}
 	}
+
+	// Update the disk size stat on load, as we can't really know it when saving
+	HLODStats.Add(FWorldPartitionHLODStats::MemoryDiskSizeBytes, FHLODActorDesc::GetPackageSize(this));
 }
 
 void AWorldPartitionHLOD::RerunConstructionScripts()
diff --git a/Engine/Source/Runtime/Engine/Private/WorldPartition/HLOD/HLODActorDesc.cpp b/Engine/Source/Runtime/Engine/Private/WorldPartition/HLOD/HLODActorDesc.cpp
--- a/Engine/Source/Runtime/Engine/Private/WorldPartition/HLOD/HLODActorDesc.cpp
+++ b/Engine/Source/Runtime/Engine/Private/WorldPartition/HLOD/HLODActorDesc.cpp
@@ -5,7 +5,7 @@
 #if WITH_EDITOR
 #include "Algo/Transform.h"
 #include "Algo/RemoveIf.h"
-#include "Hash/CityHashHelpers.h"
+#include "HAL/FileManager.h"
 #include "UObject/UE5MainStreamObjectVersion.h"
 #include "UObject/UE5ReleaseStreamObjectVersion.h"
 #include "UObject/FortniteMainBranchObjectVersion.h"
@@ -20,19 +20,10 @@ void FHLODActorDesc::Init(const AActor* InActor)
 
 	HLODSubActors.Reserve(HLODActor->GetSubActors().Num());
 	Algo::Transform(HLODActor->GetSubActors(), HLODSubActors, [](const FHLODSubActor& SubActor) { return FHLODSubActorDesc(SubActor.ActorGuid, SubActor.ContainerID); });
-	
-	CellHash = 0;
-	if (const UHLODLayer* SubActorsHLODLayer = HLODActor->GetSubActorsHLODLayer())
-	{
-		uint64 GridIndexX;
-		uint64 GridIndexY;
-		uint64 GridIndexZ;
-		HLODActor->GetGridIndices(GridIndexX, GridIndexY, GridIndexZ);
 
-		FDataLayersID DataLayersID(HLODActor->GetDataLayerInstances());
-
-		CellHash = ComputeCellHash(SubActorsHLODLayer->GetName(), GridIndexX, GridIndexY, GridIndexZ, DataLayersID);
-	}
+	SourceCellName = HLODActor->GetSourceCellName();
+	SourceHLODLayerName = HLODActor->GetSubActorsHLODLayer()->GetFName();
+	HLODStats = HLODActor->GetStats();
 }
 
 void FHLODActorDesc::Serialize(FArchive& Ar)
@@ -59,33 +50,62 @@ void FHLODActorDesc::Serialize(FArchive& Ar)
 		Ar << HLODLayer_Deprecated;
 	}
 
+	const bool bSerializeCellHash = Ar.CustomVer(FUE5MainStreamObjectVersion::GUID) >= FUE5MainStreamObjectVersion::WorldPartitionHLODActorDescSerializeCellHash &&
+									Ar.CustomVer(FFortniteMainBranchObjectVersion::GUID) < FFortniteMainBranchObjectVersion::WorldPartitionHLODActorDescSerializeStats;
 	if (Ar.CustomVer(FUE5MainStreamObjectVersion::GUID) >= FUE5MainStreamObjectVersion::WorldPartitionHLODActorDescSerializeCellHash)
 	{
-		Ar << CellHash;
+		uint64 CellHash_Deprecated;
+		Ar << CellHash_Deprecated;
 	}
 
 	if (Ar.CustomVer(FFortniteMainBranchObjectVersion::GUID) < FFortniteMainBranchObjectVersion::WorldPartitionActorDescSerializeActorIsRuntimeOnly)
 	{
 		bActorIsRuntimeOnly = true;
 	}
+
+	if (Ar.CustomVer(FFortniteMainBranchObjectVersion::GUID) >= FFortniteMainBranchObjectVersion::WorldPartitionHLODActorDescSerializeStats)
+	{		
+		Ar << SourceCellName;
+		Ar << SourceHLODLayerName;
+		Ar << HLODStats;
+
+		// Update package size stat on load
+		if (Ar.IsLoading())
+		{
+			HLODStats.Add(FWorldPartitionHLODStats::MemoryDiskSizeBytes, GetPackageSize());
+		}
+	}
 }
 
 bool FHLODActorDesc::Equals(const FWorldPartitionActorDesc* Other) const
 {
 	if (FWorldPartitionActorDesc::Equals(Other))
 	{
-		const FHLODActorDesc* HLODActorDesc = (FHLODActorDesc*)Other;
-		return (CellHash == HLODActorDesc->CellHash) && CompareUnsortedArrays(HLODSubActors, HLODActorDesc->HLODSubActors);
+		const FHLODActorDesc& HLODActorDesc = *(FHLODActorDesc*)Other;
+		return SourceCellName == HLODActorDesc.SourceCellName &&
+			   SourceHLODLayerName == HLODActorDesc.SourceHLODLayerName &&
+			   HLODStats.OrderIndependentCompareEqual(HLODActorDesc.GetStats()) &&
+			   CompareUnsortedArrays(HLODSubActors, HLODActorDesc.HLODSubActors);
 	}
 	return false;
 }
 
-uint64 FHLODActorDesc::ComputeCellHash(const FString HLODLayerName, uint64 GridIndexX, uint64 GridIndexY, uint64 GridIndexZ, FDataLayersID DataLayersID)
+static int64 GetPackageSize(const FString& InPackageFileName)
 {
-	uint64 CellHash = FCrc::StrCrc32(*HLODLayerName);
-	CellHash = AppendCityHash(GridIndexX, CellHash);
-	CellHash = AppendCityHash(GridIndexY, CellHash);
-	CellHash = AppendCityHash(GridIndexZ, CellHash);
-	return HashCombine(DataLayersID.GetHash(), CellHash);
+	const int64 PackageSize = IFileManager::Get().FileSize(*InPackageFileName);
+	return PackageSize;
 }
+
+int64 FHLODActorDesc::GetPackageSize() const
+{
+	const FString PackageFileName = FPackageName::LongPackageNameToFilename(GetActorPackage().ToString(), FPackageName::GetAssetPackageExtension());
+	return ::GetPackageSize(PackageFileName);
+}
+
+int64 FHLODActorDesc::GetPackageSize(const AWorldPartitionHLOD* InHLODActor)
+{
+	const FString PackageFileName = InHLODActor->GetPackage()->GetLoadedPath().GetLocalFullPath();
+	return ::GetPackageSize(PackageFileName);
+}
+
 #endif
diff --git a/Engine/Source/Runtime/Engine/Private/WorldPartition/HLOD/HLODStats.cpp b/Engine/Source/Runtime/Engine/Private/WorldPartition/HLOD/HLODStats.cpp
new file mode 100644
--- /dev/null
+++ b/Engine/Source/Runtime/Engine/Private/WorldPartition/HLOD/HLODStats.cpp
@@ -0,0 +1,30 @@
+// Copyright Epic Games, Inc. All Rights Reserved.
+
+#include "WorldPartition/HLOD/HLODStats.h"
+
+
+const FName FWorldPartitionHLODStats::InputActorCount(TEXT("InputActorCount"));
+const FName FWorldPartitionHLODStats::InputTriangleCount(TEXT("InputTriangleCount"));
+const FName FWorldPartitionHLODStats::InputVertexCount(TEXT("InputVertexCount"));
+
+const FName FWorldPartitionHLODStats::MeshInstanceCount(TEXT("MeshInstanceCount"));
+const FName FWorldPartitionHLODStats::MeshNaniteTriangleCount(TEXT("MeshNaniteTriangleCount"));
+const FName FWorldPartitionHLODStats::MeshNaniteVertexCount(TEXT("MeshNaniteVertexCount"));
+const FName FWorldPartitionHLODStats::MeshTriangleCount(TEXT("MeshTriangleCount"));
+const FName FWorldPartitionHLODStats::MeshVertexCount(TEXT("MeshVertexCount"));
+const FName FWorldPartitionHLODStats::MeshUVChannelCount(TEXT("MeshUVChannelCount"));
+
+const FName FWorldPartitionHLODStats::MaterialBaseColorTextureSize(TEXT("MaterialBaseColorTextureSize"));
+const FName FWorldPartitionHLODStats::MaterialNormalTextureSize(TEXT("MaterialNormalTextureSize"));
+const FName FWorldPartitionHLODStats::MaterialEmissiveTextureSize(TEXT("MaterialEmissiveTextureSize"));
+const FName FWorldPartitionHLODStats::MaterialMetallicTextureSize(TEXT("MaterialMetallicTextureSize"));
+const FName FWorldPartitionHLODStats::MaterialRoughnessTextureSize(TEXT("MaterialRoughnessTextureSize"));
+const FName FWorldPartitionHLODStats::MaterialSpecularTextureSize(TEXT("MaterialSpecularTextureSize"));
+
+const FName FWorldPartitionHLODStats::MemoryMeshResourceSizeBytes(TEXT("MemoryMeshResourceSizeBytes"));
+const FName FWorldPartitionHLODStats::MemoryTexturesResourceSizeBytes(TEXT("MemoryTexturesResourceSizeBytes"));
+const FName FWorldPartitionHLODStats::MemoryDiskSizeBytes(TEXT("MemoryDiskSizeBytes"));
+
+const FName FWorldPartitionHLODStats::BuildTimeLoadMilliseconds(TEXT("BuildTimeLoadMilliseconds"));
+const FName FWorldPartitionHLODStats::BuildTimeBuildMilliseconds(TEXT("BuildTimeBuildMilliseconds"));
+const FName FWorldPartitionHLODStats::BuildTimeTotalMilliseconds(TEXT("BuildTimeTotalMilliseconds"));
diff --git a/Engine/Source/Runtime/Engine/Private/WorldPartition/HLOD/HLODSubsystem.cpp b/Engine/Source/Runtime/Engine/Private/WorldPartition/HLOD/HLODSubsystem.cpp
--- a/Engine/Source/Runtime/Engine/Private/WorldPartition/HLOD/HLODSubsystem.cpp
+++ b/Engine/Source/Runtime/Engine/Private/WorldPartition/HLOD/HLODSubsystem.cpp
@@ -4,6 +4,8 @@
 #include "WorldPartition/HLOD/HLODActor.h"
 #include "WorldPartition/HLOD/HLODActorDesc.h"
 #include "WorldPartition/HLOD/HLODLayer.h"
+#include "WorldPartition/DataLayer/DataLayerInstance.h"
+#include "WorldPartition/DataLayer/DataLayerSubsystem.h"
 #include "WorldPartition/WorldPartition.h"
 #include "WorldPartition/WorldPartitionSubsystem.h"
 #include "WorldPartition/WorldPartitionRuntimeCell.h"
@@ -18,6 +20,9 @@
 #include "EngineUtils.h"
 #include "LevelUtils.h"
 #include "Components/StaticMeshComponent.h"
+#include "HAL/FileManager.h"
+#include "Misc/FileHelper.h"
+#include "Templates/UniquePtr.h"
 
 #include UE_INLINE_GENERATED_CPP_BY_NAME(HLODSubsystem)
 
@@ -469,5 +474,101 @@ void FHLODResourcesResidencySceneViewExtension::BeginRenderViewFamily(FSceneView
 	GetWorld()->GetSubsystem<UHLODSubsystem>()->OnBeginRenderViews(InViewFamily);
 }
 
+#if WITH_EDITOR
+
+static void DumpHLODStats(UWorld* World)
+{
+	UWorldPartition* WorldPartition = World ? World->GetWorldPartition() : nullptr;
+	if (!WorldPartition)
+	{
+		return;
+	}
+	
+	typedef TFunction<FString(const FHLODActorDesc&)> FGetStatFunc;
+
+	auto GetHLODStat = [](FName InStatName)
+	{
+		return TPair<FName, FGetStatFunc>(InStatName, [InStatName](const FHLODActorDesc& InActorDesc)
+		{
+			return FString::Printf(TEXT("%lld"), InActorDesc.GetStat(InStatName));
+		});
+	};
+
+	const UDataLayerSubsystem* DataLayerSubsystem = UWorld::GetSubsystem<UDataLayerSubsystem>(World);
+
+	auto GetDataLayerShortName = [DataLayerSubsystem](FName DataLayerInstanceName)
+	{ 
+		UDataLayerInstance* DataLayerInstance = DataLayerSubsystem ? DataLayerSubsystem->GetDataLayerInstance(DataLayerInstanceName) : nullptr;
+		return DataLayerInstance ? DataLayerInstance->GetDataLayerShortName() : DataLayerInstanceName.ToString();
+	};
+
+	TArray<TPair<FName, FGetStatFunc>> StatsToWrite =
+	{
+		{ "Name",				[](const FHLODActorDesc& InActorDesc) { return InActorDesc.GetActorLabel().ToString(); } },
+		{ "HLODLayer",			[](const FHLODActorDesc& InActorDesc) { return InActorDesc.GetSourceHLODLayerName().ToString(); } },
+		{ "SpatiallyLoaded",	[](const FHLODActorDesc& InActorDesc) { return InActorDesc.GetIsSpatiallyLoaded() ? TEXT("true") : TEXT("false"); } },
+		{ "DataLayers",			[&GetDataLayerShortName](const FHLODActorDesc& InActorDesc) { return FString::JoinBy(InActorDesc.GetDataLayerInstanceNames(), TEXT(" | "), GetDataLayerShortName); } },
+
+		GetHLODStat(FWorldPartitionHLODStats::InputActorCount),
+		GetHLODStat(FWorldPartitionHLODStats::InputTriangleCount),
+		GetHLODStat(FWorldPartitionHLODStats::InputVertexCount),
+
+		GetHLODStat(FWorldPartitionHLODStats::MeshInstanceCount),
+		GetHLODStat(FWorldPartitionHLODStats::MeshNaniteTriangleCount),
+		GetHLODStat(FWorldPartitionHLODStats::MeshNaniteVertexCount),
+		GetHLODStat(FWorldPartitionHLODStats::MeshTriangleCount),
+		GetHLODStat(FWorldPartitionHLODStats::MeshVertexCount),
+		GetHLODStat(FWorldPartitionHLODStats::MeshUVChannelCount),
+
+		GetHLODStat(FWorldPartitionHLODStats::MaterialBaseColorTextureSize),
+		GetHLODStat(FWorldPartitionHLODStats::MaterialNormalTextureSize),
+		GetHLODStat(FWorldPartitionHLODStats::MaterialEmissiveTextureSize),
+		GetHLODStat(FWorldPartitionHLODStats::MaterialMetallicTextureSize),
+		GetHLODStat(FWorldPartitionHLODStats::MaterialRoughnessTextureSize),
+		GetHLODStat(FWorldPartitionHLODStats::MaterialSpecularTextureSize),
+		
+		GetHLODStat(FWorldPartitionHLODStats::MemoryMeshResourceSizeBytes),
+		GetHLODStat(FWorldPartitionHLODStats::MemoryTexturesResourceSizeBytes),
+		GetHLODStat(FWorldPartitionHLODStats::MemoryDiskSizeBytes),
+		
+		GetHLODStat(FWorldPartitionHLODStats::BuildTimeLoadMilliseconds),
+		GetHLODStat(FWorldPartitionHLODStats::BuildTimeBuildMilliseconds),
+		GetHLODStat(FWorldPartitionHLODStats::BuildTimeTotalMilliseconds)				
+	};
+
+	// Write header
+	FStringOutputDevice Output;
+	const FString StatHeader = FString::JoinBy(StatsToWrite, TEXT(","), [](const TPair<FName, FGetStatFunc>& Pair) { return Pair.Key.ToString(); });
+	Output.Logf(TEXT("%s" LINE_TERMINATOR_ANSI), *StatHeader);
+
+	// Write one line per HLOD actor desc
+	for (FActorDescList::TIterator<AWorldPartitionHLOD> HLODIterator(WorldPartition->GetActorDescContainer()); HLODIterator; ++HLODIterator)
+	{
+		const FString StatLine = FString::JoinBy(StatsToWrite, TEXT(","), [&HLODIterator](const TPair<FName, FGetStatFunc>& Pair) { return Pair.Value(**HLODIterator); });
+		Output.Logf(TEXT("%s" LINE_TERMINATOR_ANSI), *StatLine);
+	}
+
+	// Write to file
+	const FString HLODStatsOutputFilename = FPaths::ProjectSavedDir() / TEXT("WorldPartition") / FString::Printf(TEXT("HLODStats-%s-%08x-%s.csv"), *World->GetName(), FPlatformProcess::GetCurrentProcessId(), *FDateTime::Now().ToString());
+	FFileHelper::SaveStringToFile(Output, *HLODStatsOutputFilename);
+}
+
+FAutoConsoleCommand HLODDumpStats(
+	TEXT("wp.Editor.HLOD.DumpStats"),
+	TEXT("Write various HLOD stats to a CSV formatted file."),
+	FConsoleCommandWithArgsDelegate::CreateLambda([](const TArray<FString>& Args)
+	{
+		for (const FWorldContext& Context : GEngine->GetWorldContexts())
+		{
+			if (UWorld* World = Context.World())
+			{
+				DumpHLODStats(World);
+			}
+		}
+	})
+);
+
+#endif // #if WITH_EDITOR
+
 #undef LOCTEXT_NAMESPACE
 
diff --git a/Engine/Source/Runtime/Engine/Private/WorldPartition/RuntimeSpatialHash/RuntimeSpatialHashHLOD.cpp b/Engine/Source/Runtime/Engine/Private/WorldPartition/RuntimeSpatialHash/RuntimeSpatialHashHLOD.cpp
--- a/Engine/Source/Runtime/Engine/Private/WorldPartition/RuntimeSpatialHash/RuntimeSpatialHashHLOD.cpp
+++ b/Engine/Source/Runtime/Engine/Private/WorldPartition/RuntimeSpatialHash/RuntimeSpatialHashHLOD.cpp
@@ -205,10 +205,6 @@ static TArray<FGuid> GenerateHLODsForGrid(UWorldPartition* WorldPartition, const
 
 				FHLODCreationParams CreationParams;
 				CreationParams.WorldPartition = WorldPartition;
-				CreationParams.GridIndexX = CellCoord.X;
-				CreationParams.GridIndexY = CellCoord.Y;
-				CreationParams.GridIndexZ = CellCoord.Z;
-				CreationParams.DataLayersID = GridCellDataChunk.GetDataLayersID();
 				CreationParams.CellName = FName(CellName);
 				CreationParams.CellBounds = CellBounds;
 				CreationParams.HLODLevel = HLODLevel;
@@ -428,35 +424,10 @@ bool UWorldPartitionRuntimeSpatialHash::GenerateHLOD(ISourceControlHelper* Sourc
 
 	// HLOD creation context
 	FHLODCreationContext Context;
-
-	TSet<uint64> InvalidHLODCellHashes;
-	InvalidHLODCellHashes.Add(0);
-
-	TSet<FWorldPartitionHandle> InvalidHLODActors;
-
 	for (FActorDescList::TIterator<AWorldPartitionHLOD> HLODIterator(WorldPartition->GetActorDescContainer()); HLODIterator; ++HLODIterator)
 	{
-		uint64 CellHash = HLODIterator->GetCellHash();
-
 		FWorldPartitionHandle HLODActorHandle(WorldPartition, HLODIterator->GetGuid());
-		FWorldPartitionHandle* DupeHLODActorHandle = Context.HLODActorDescs.Find(CellHash);
-
-		if (DupeHLODActorHandle == nullptr && !InvalidHLODCellHashes.Contains(CellHash))
-		{
-			Context.HLODActorDescs.Emplace(CellHash, MoveTemp(HLODActorHandle));
-		}
-		else
-		{
-			InvalidHLODCellHashes.Add(CellHash);
-
-			InvalidHLODActors.Emplace(MoveTemp(HLODActorHandle));
-			if (DupeHLODActorHandle)
-			{
-				InvalidHLODActors.Emplace(*DupeHLODActorHandle);
-			}
-
-			Context.HLODActorDescs.Remove(CellHash);
-		}
+		Context.HLODActorDescs.Emplace(HLODIterator->GetActorName(), MoveTemp(HLODActorHandle));
 	}
 
 	IStreamingGenerationContext::FActorSetContainer* MainActorSetContainer = const_cast<IStreamingGenerationContext::FActorSetContainer*>(StreamingGenerationContext->GetMainWorldContainer());
@@ -562,9 +533,6 @@ bool UWorldPartitionRuntimeSpatialHash::GenerateHLOD(ISourceControlHelper* Sourc
 		DeleteHLODActor(HLODActorPair.Value);
 	}
 
-	// Destroy all invalid HLOD actors
-	Algo::ForEach(InvalidHLODActors, DeleteHLODActor);
-
 	// Create/destroy HLOD grid actors
 	UpdateHLODGridsActors(GetWorld(), HLODGrids, SourceControlHelper);
 
diff --git a/Engine/Source/Runtime/Engine/Public/WorldPartition/HLOD/HLODActor.h b/Engine/Source/Runtime/Engine/Public/WorldPartition/HLOD/HLODActor.h
--- a/Engine/Source/Runtime/Engine/Public/WorldPartition/HLOD/HLODActor.h
+++ b/Engine/Source/Runtime/Engine/Public/WorldPartition/HLOD/HLODActor.h
@@ -7,6 +7,7 @@
 #include "WorldPartition/WorldPartitionRuntimeCell.h"
 #include "WorldPartition/WorldPartitionActorDesc.h"
 #include "WorldPartition/HLOD/HLODSubActor.h"
+#include "WorldPartition/HLOD/HLODStats.h"
 #include "HLODActor.generated.h"
 
 class UHLODLayer;
@@ -17,6 +18,8 @@ class ENGINE_API AWorldPartitionHLOD : public AActor
 {
 	GENERATED_UCLASS_BODY()
 
+	typedef TMap<FName, int64>	FStats;
+
 public:
 	void SetVisibility(bool bInVisible);
 
@@ -38,20 +41,6 @@ public:
 
 	void SetRequireWarmup(bool InRequireWarmup) { bRequireWarmup = InRequireWarmup; }
 
-	void SetGridIndices(uint64 InGridIndexX, uint64 InGridIndexY, uint64 InGridIndexZ)
-	{
-		GridIndexX = InGridIndexX;
-		GridIndexY = InGridIndexY;
-		GridIndexZ = InGridIndexZ;
-	}
-
-	void GetGridIndices(uint64& OutGridIndexX, uint64& OutGridIndexY, uint64& OutGridIndexZ) const
-	{
-		OutGridIndexX = GridIndexX;
-		OutGridIndexY = GridIndexY;
-		OutGridIndexZ = GridIndexZ;
-	}
-
 	void SetSourceCellName(FName InSourceCellName);
 	inline void SetLODLevel(uint32 InLODLevel) { LODLevel = InLODLevel; }
 
@@ -63,6 +52,11 @@ public:
 
 	void BuildHLOD(bool bForceBuild = false);
 	uint32 GetHLODHash() const;
+
+	const FStats& GetStats() const { return HLODStats; }
+	int64 GetStat(FName InStatName) const { return HLODStats.FindRef(InStatName); }
+	void SetStat(FName InStatName, int64 InStatValue) { HLODStats.Add(InStatName, InStatValue); }
+	void ResetStats() { HLODStats.Reset(); }
 #endif // WITH_EDITOR
 
 protected:
@@ -100,15 +94,6 @@ private:
 	UPROPERTY()
 	TObjectPtr<const UHLODLayer> SubActorsHLODLayer;
 
-	UPROPERTY()
-	int64 GridIndexX;
-
-	UPROPERTY()
-	int64 GridIndexY;
-
-	UPROPERTY()
-	int64 GridIndexZ;
-
 	UPROPERTY()
 	FBox HLODBounds;
 
@@ -117,6 +102,9 @@ private:
 
 	UPROPERTY()
 	uint32 HLODHash;
+
+	UPROPERTY()
+	TMap<FName, int64> HLODStats;
 #endif
 
 	UPROPERTY()
diff --git a/Engine/Source/Runtime/Engine/Public/WorldPartition/HLOD/HLODActorDesc.h b/Engine/Source/Runtime/Engine/Public/WorldPartition/HLOD/HLODActorDesc.h
--- a/Engine/Source/Runtime/Engine/Public/WorldPartition/HLOD/HLODActorDesc.h
+++ b/Engine/Source/Runtime/Engine/Public/WorldPartition/HLOD/HLODActorDesc.h
@@ -9,9 +9,11 @@
 
 #include "WorldPartition/WorldPartitionActorDesc.h"
 #include "WorldPartition/DataLayer/DataLayersID.h"
+#include "WorldPartition/HLOD/HLODStats.h"
 #include "WorldPartition/HLOD/HLODSubActor.h"
 
 class UHLODLayer;
+class AWorldPartitionHLOD;
 
 /**
  * ActorDesc for AWorldPartitionHLOD.
@@ -21,10 +23,16 @@ class ENGINE_API FHLODActorDesc : public FWorldPartitionActorDesc
 	friend class FHLODActorDescFactory;
 
 public:
+	typedef TMap<FName, int64>	FStats;
+
 	inline const TArray<FHLODSubActorDesc>& GetSubActors() const { return HLODSubActors; }
-	inline uint64 GetCellHash() const { return CellHash; }
+	inline const FName GetSourceCellName() const { return SourceCellName; }
+	inline const FName GetSourceHLODLayerName() const { return SourceHLODLayerName; }
+	inline const FStats& GetStats() const { return HLODStats; }
+	inline int64 GetStat(FName InStatName) const { return HLODStats.FindRef(InStatName); }
 
-	static uint64 ComputeCellHash(const FString HLODLayerName, uint64 GridIndexX, uint64 GridIndexY, uint64 GridIndexZ, FDataLayersID DataLayersID);
+	int64 GetPackageSize() const;
+	static int64 GetPackageSize(const AWorldPartitionHLOD* InHLODActor);
 
 protected:
 	//~ Begin FWorldPartitionActorDesc Interface.
@@ -35,7 +43,8 @@ protected:
 	//~ End FWorldPartitionActorDesc Interface.
 
 	TArray<FHLODSubActorDesc> HLODSubActors;
-
-	uint64 CellHash = 0;
+	FName SourceCellName;
+	FName SourceHLODLayerName;
+	FStats HLODStats;
 };
 #endif
diff --git a/Engine/Source/Runtime/Engine/Public/WorldPartition/HLOD/HLODStats.h b/Engine/Source/Runtime/Engine/Public/WorldPartition/HLOD/HLODStats.h
new file mode 100644
--- /dev/null
+++ b/Engine/Source/Runtime/Engine/Public/WorldPartition/HLOD/HLODStats.h
@@ -0,0 +1,40 @@
+// Copyright Epic Games, Inc. All Rights Reserved.
+
+#pragma once
+
+#include "CoreMinimal.h"
+
+
+struct ENGINE_API FWorldPartitionHLODStats
+{
+	// Input
+	static const FName InputActorCount;
+	static const FName InputTriangleCount;
+	static const FName InputVertexCount;
+
+	// Mesh
+	static const FName MeshInstanceCount;
+	static const FName MeshNaniteTriangleCount;
+	static const FName MeshNaniteVertexCount;
+	static const FName MeshTriangleCount;
+	static const FName MeshVertexCount;
+	static const FName MeshUVChannelCount;
+
+	// Material
+	static const FName MaterialBaseColorTextureSize;
+	static const FName MaterialNormalTextureSize;
+	static const FName MaterialEmissiveTextureSize;
+	static const FName MaterialMetallicTextureSize;
+	static const FName MaterialRoughnessTextureSize;
+	static const FName MaterialSpecularTextureSize;
+	
+	// Memory
+	static const FName MemoryMeshResourceSizeBytes;
+	static const FName MemoryTexturesResourceSizeBytes;
+	static const FName MemoryDiskSizeBytes;
+
+	// Build
+	static const FName BuildTimeLoadMilliseconds;
+	static const FName BuildTimeBuildMilliseconds;
+	static const FName BuildTimeTotalMilliseconds;
+};
diff --git a/Engine/Source/Runtime/Engine/Public/WorldPartition/HLOD/IWorldPartitionHLODUtilities.h b/Engine/Source/Runtime/Engine/Public/WorldPartition/HLOD/IWorldPartitionHLODUtilities.h
--- a/Engine/Source/Runtime/Engine/Public/WorldPartition/HLOD/IWorldPartitionHLODUtilities.h
+++ b/Engine/Source/Runtime/Engine/Public/WorldPartition/HLOD/IWorldPartitionHLODUtilities.h
@@ -17,7 +17,7 @@ class AWorldPartitionHLOD;
 
 struct ENGINE_API FHLODCreationContext
 {
-	TMap<uint64, FWorldPartitionHandle> HLODActorDescs;
+	TMap<FName, FWorldPartitionHandle> HLODActorDescs;
 	TArray<FWorldPartitionReference> ActorReferences;
 };
 
@@ -25,11 +25,6 @@ struct ENGINE_API FHLODCreationParams
 {
 	UWorldPartition* WorldPartition;
 
-	// Everything needed to build the cell hash
-	int64 GridIndexX;
-	int64 GridIndexY;
-	int64 GridIndexZ;
-	FDataLayersID DataLayersID;
 	FName HLODLayerName;
 
 	FName CellName;
