World Partition - HLOD: Added support for runtime HLOD modification
* An HLOD modifier class can be assigned to an HLOD Layer
* HLOD Modifier classes can inject new components during an HLOD build

diff --git a/Engine/Plugins/Editor/WorldPartitionHLODUtilities/Source/Private/WorldPartitionHLODUtilities.cpp b/Engine/Plugins/Editor/WorldPartitionHLODUtilities/Source/Private/WorldPartitionHLODUtilities.cpp
--- a/Engine/Plugins/Editor/WorldPartitionHLODUtilities/Source/Private/WorldPartitionHLODUtilities.cpp
+++ b/Engine/Plugins/Editor/WorldPartitionHLODUtilities/Source/Private/WorldPartitionHLODUtilities.cpp
@@ -13,6 +13,7 @@
 #include "WorldPartition/HLOD/HLODActorDesc.h"
 #include "WorldPartition/HLOD/HLODBuilder.h"
 #include "WorldPartition/HLOD/HLODLayer.h"
+#include "WorldPartition/HLOD/HLODModifier.h"
 #include "WorldPartition/WorldPartitionLevelStreamingDynamic.h"
 
 #include "HLODBuilderInstancing.h"
@@ -86,7 +87,7 @@ static uint32 ComputeHLODHash(AWorldPartitionHLOD* InHLODActor, const TArray<AAc
 	FArchiveCrc32 Ar;
 
 	// Base key, changing this will force a rebuild of all HLODs
-	FString HLODBaseKey = "5052091956924DB3BD9ACE00B71944AC";
+	FString HLODBaseKey = "A6C1517FC49B478E833A540E2C06289F";
 	Ar << HLODBaseKey;
 
 	// HLOD Layer
@@ -109,41 +110,39 @@ static uint32 ComputeHLODHash(AWorldPartitionHLOD* InHLODActor, const TArray<AAc
 
 TArray<AWorldPartitionHLOD*> FWorldPartitionHLODUtilities::CreateHLODActors(FHLODCreationContext& InCreationContext, const FHLODCreationParams& InCreationParams, const TArray<IStreamingGenerationContext::FActorInstance>& InActors, const TArray<const UDataLayerInstance*>& InDataLayersInstances)
 {
-	struct FSubActorsInfo
-	{
-		TArray<FHLODSubActor>	SubActors;
-		bool					bIsSpatiallyLoaded;
-	};
-	TMap<UHLODLayer*, FSubActorsInfo> SubActorsInfos;
+	TMap<UHLODLayer*, TArray<FHLODSubActor>> SubActorsPerHLODLayer;
 
 	for (const IStreamingGenerationContext::FActorInstance& ActorInstance : InActors)
 	{
 		const FWorldPartitionActorDescView& ActorDescView = ActorInstance.GetActorDescView();
+
+		if (!ActorDescView.GetIsSpatiallyLoaded())
+		{
+			UE_LOG(LogHLODBuilder, Error, TEXT("Tried to included non-spatially loaded actor %s into HLOD"), *ActorDescView.GetActorName().ToString());
+			continue;
+		}
+
 		if (ActorDescView.GetActorIsHLODRelevant())
 		{
 			UHLODLayer* HLODLayer = UHLODLayer::GetHLODLayer(ActorDescView, InCreationParams.WorldPartition);
 			if (HLODLayer)
 			{
-				FSubActorsInfo& SubActorsInfo = SubActorsInfos.FindOrAdd(HLODLayer);
+				TArray<FHLODSubActor>& SubActors = SubActorsPerHLODLayer.FindOrAdd(HLODLayer);
 
 				// Leaving this as deprecated for now until we fix up the serialization format for SubActorsInfo and do proper upgrade/deprecation 
 PRAGMA_DISABLE_DEPRECATION_WARNINGS
-				SubActorsInfo.SubActors.Emplace(ActorDescView.GetGuid(), ActorDescView.GetActorPackage(), ActorDescView.GetActorPath(), ActorInstance.GetContainerID(), ActorInstance.GetActorDescContainer()->GetContainerPackage(), ActorInstance.GetTransform());
+				SubActors.Emplace(ActorDescView.GetGuid(), ActorDescView.GetActorPackage(), ActorDescView.GetActorPath(), ActorInstance.GetContainerID(), ActorInstance.GetActorDescContainer()->GetContainerPackage(), ActorInstance.GetTransform());
 PRAGMA_ENABLE_DEPRECATION_WARNINGS
-				if (ActorDescView.GetIsSpatiallyLoaded())
-				{
-					SubActorsInfo.bIsSpatiallyLoaded = true;
-				}
 			}
 		}
 	}
 
 	TArray<AWorldPartitionHLOD*> HLODActors;
-	for (const auto& Pair : SubActorsInfos)
+	for (const auto& Pair : SubActorsPerHLODLayer)
 	{
 		const UHLODLayer* HLODLayer = Pair.Key;
-		const FSubActorsInfo& SubActorsInfo = Pair.Value;
-		check(!SubActorsInfo.SubActors.IsEmpty());
+		const TArray<FHLODSubActor>& SubActors = Pair.Value;
+		check(!SubActors.IsEmpty());
 
 		auto ComputeCellHash = [](const FString& HLODLayerName, const FName CellName)
 		{
@@ -192,11 +191,11 @@ PRAGMA_ENABLE_DEPRECATION_WARNINGS
 
 		// Sub actors
 		{
-			bool bSubActorsChanged = HLODActor->GetSubActors().Num() != SubActorsInfo.SubActors.Num();
+			bool bSubActorsChanged = HLODActor->GetSubActors().Num() != SubActors.Num();
 			if (!bSubActorsChanged)
 			{
 				TArray<FHLODSubActor> A = HLODActor->GetSubActors();
-				TArray<FHLODSubActor> B = SubActorsInfo.SubActors;
+				TArray<FHLODSubActor> B = SubActors;
 				A.Sort();
 				B.Sort();
 				bSubActorsChanged = A != B;
@@ -204,7 +203,7 @@ PRAGMA_ENABLE_DEPRECATION_WARNINGS
 
 			if (bSubActorsChanged)
 			{
-				HLODActor->SetSubActors(SubActorsInfo.SubActors);
+				HLODActor->SetSubActors(SubActors);
 				bIsDirty = true;
 			}
 		}
@@ -218,11 +217,9 @@ PRAGMA_ENABLE_DEPRECATION_WARNINGS
 		}
 
 		// Spatially loaded
-		// HLOD that are always loaded will not take the SubActorsInfo.GridPlacement into account
-		bool bExpectedIsSpatiallyLoaded = !HLODLayer->IsSpatiallyLoaded() ? false : SubActorsInfo.bIsSpatiallyLoaded;
-		if (HLODActor->GetIsSpatiallyLoaded() != bExpectedIsSpatiallyLoaded)
+		if (HLODActor->GetIsSpatiallyLoaded() != HLODLayer->IsSpatiallyLoaded())
 		{
-			HLODActor->SetIsSpatiallyLoaded(bExpectedIsSpatiallyLoaded);
+			HLODActor->SetIsSpatiallyLoaded(HLODLayer->IsSpatiallyLoaded());
 			bIsDirty = true;
 		}
 
@@ -589,24 +586,40 @@ uint32 FWorldPartitionHLODUtilities::BuildHLOD(AWorldPartitionHLOD* InHLODActor)
 		UHLODBuilder* HLODBuilder = NewObject<UHLODBuilder>(GetTransientPackage(), HLODBuilderClass);
 		if (ensure(HLODBuilder))
 		{
-			FGCObjectScopeGuard BuilderGCScopeGuard(HLODBuilder);
+			FGCObjectScopeGuard HLODBuilderGCScopeGuard(HLODBuilder);
 
 			HLODBuilder->SetHLODBuilderSettings(HLODLayer->GetHLODBuilderSettings());
 
 			FHLODBuildContext HLODBuildContext;
 			HLODBuildContext.World = InHLODActor->GetWorld();
+			HLODBuildContext.SourceActors = LevelStreaming->GetLoadedLevel()->Actors;
 			HLODBuildContext.AssetsOuter = InHLODActor->GetPackage();
 			HLODBuildContext.AssetsBaseName = InHLODActor->GetActorLabel();
 			HLODBuildContext.MinVisibleDistance = InHLODActor->GetMinVisibleDistance();
+			HLODBuildContext.WorldPosition = InHLODActor->GetHLODBounds().GetCenter();
+
+			const TSubclassOf<UWorldPartitionHLODModifier> HLODModifierClass = HLODLayer->GetHLODModifierClass();
+			UWorldPartitionHLODModifier* HLODModifier = HLODModifierClass.Get() ? NewObject<UWorldPartitionHLODModifier>(GetTransientPackage(), HLODModifierClass) : nullptr;
+			FGCObjectScopeGuard HLODModifierGCScopeGuard(HLODModifier);
+
+			if (HLODModifier)
+			{
+				HLODModifier->BeginHLODBuild(HLODBuildContext);
+			}
 
 			// Build
 			TArray<UActorComponent*> HLODComponents;
 			{
 				FAutoScopedDurationTimer BuildTimeScope;
-				HLODComponents = HLODBuilder->Build(HLODBuildContext, LevelStreaming->GetLoadedLevel()->Actors);
+				HLODComponents = HLODBuilder->Build(HLODBuildContext);
 				BuildTimeMS = BuildTimeScope.GetTime() * 1000;
 			}
 
+			if (HLODModifier)
+			{
+				HLODModifier->EndHLODBuild(HLODComponents);
+			}
+
 			if (HLODComponents.IsEmpty())
 			{
 				UE_LOG(LogHLODBuilder, Warning, TEXT("HLOD generation created no component for %s"), *InHLODActor->GetActorLabel());
@@ -634,6 +647,9 @@ uint32 FWorldPartitionHLODUtilities::BuildHLOD(AWorldPartitionHLOD* InHLODActor)
 					PrimitiveComponent->SetGenerateOverlapEvents(false);
 					PrimitiveComponent->CanCharacterStepUpOn = ECanBeCharacterBase::ECB_No;
 					PrimitiveComponent->SetCollisionEnabled(ECollisionEnabled::NoCollision);
+
+					// HLOD visual components aren't needed on servers
+					PrimitiveComponent->AlwaysLoadOnServer = false;
 				}
 
 				if (UStaticMeshComponent* StaticMeshComponent = Cast<UStaticMeshComponent>(HLODComponent))
diff --git a/Engine/Plugins/Editor/WorldPartitionHLODUtilities/Source/WorldPartitionHLODUtilities.Build.cs b/Engine/Plugins/Editor/WorldPartitionHLODUtilities/Source/WorldPartitionHLODUtilities.Build.cs
--- a/Engine/Plugins/Editor/WorldPartitionHLODUtilities/Source/WorldPartitionHLODUtilities.Build.cs
+++ b/Engine/Plugins/Editor/WorldPartitionHLODUtilities/Source/WorldPartitionHLODUtilities.Build.cs
@@ -17,14 +17,17 @@ public class WorldPartitionHLODUtilities : ModuleRules
             new string[]
 			{
 				"Engine",
-				"MaterialUtilities"
+				"MaterialUtilities",
+				"MeshDescription",
+				"MeshMergeUtilities",
+				"StaticMeshDescription"
 			}
         );
 
         PrivateIncludePathModuleNames.AddRange(
             new string[]
             {
-				"GeometryProcessingInterfaces",
+				"GeometryProcessingInterfaces"
 			}
         );
 
@@ -32,7 +35,6 @@ public class WorldPartitionHLODUtilities : ModuleRules
             new string[]
             {
                 "MeshUtilities",
-                "MeshMergeUtilities",
                 "MeshReductionInterface",
 				"GeometryProcessingInterfaces",
 			}
diff --git a/Engine/Source/Runtime/Engine/Private/WorldPartition/HLOD/HLODBuilder.cpp b/Engine/Source/Runtime/Engine/Private/WorldPartition/HLOD/HLODBuilder.cpp
--- a/Engine/Source/Runtime/Engine/Private/WorldPartition/HLOD/HLODBuilder.cpp
+++ b/Engine/Source/Runtime/Engine/Private/WorldPartition/HLOD/HLODBuilder.cpp
@@ -308,9 +308,9 @@ static bool ShouldBatchComponent(UActorComponent* ActorComponent)
 	return bShouldBatch;
 }
 
-TArray<UActorComponent*> UHLODBuilder::Build(const FHLODBuildContext& InHLODBuildContext, const TArray<AActor*>& InSourceActors) const
+TArray<UActorComponent*> UHLODBuilder::Build(const FHLODBuildContext& InHLODBuildContext) const
 {
-	TArray<UActorComponent*> HLODRelevantComponents = GatherHLODRelevantComponents(InSourceActors);
+	TArray<UActorComponent*> HLODRelevantComponents = GatherHLODRelevantComponents(InHLODBuildContext.SourceActors);
 	if (HLODRelevantComponents.IsEmpty())
 	{
 		return {};
diff --git a/Engine/Source/Runtime/Engine/Private/WorldPartition/HLOD/HLODLayer.cpp b/Engine/Source/Runtime/Engine/Private/WorldPartition/HLOD/HLODLayer.cpp
--- a/Engine/Source/Runtime/Engine/Private/WorldPartition/HLOD/HLODLayer.cpp
+++ b/Engine/Source/Runtime/Engine/Private/WorldPartition/HLOD/HLODLayer.cpp
@@ -5,22 +5,20 @@
 =============================================================================*/
 
 #include "WorldPartition/HLOD/HLODLayer.h"
-#include "WorldPartition/HLOD/HLODActor.h"
 
 #include UE_INLINE_GENERATED_CPP_BY_NAME(HLODLayer)
 
 #if WITH_EDITOR
-#include "Serialization/ArchiveCrc32.h"
-
 #include "Engine/World.h"
 #include "Misc/ConfigCacheIni.h"
-
+#include "Modules/ModuleManager.h"
+#include "Serialization/ArchiveCrc32.h"
+#include "WorldPartition/HLOD/HLODActor.h"
+#include "WorldPartition/HLOD/HLODBuilder.h"
+#include "WorldPartition/HLOD/IWorldPartitionHLODUtilitiesModule.h"
 #include "WorldPartition/WorldPartition.h"
 #include "WorldPartition/WorldPartitionActorDesc.h"
 #include "WorldPartition/WorldPartitionActorDescView.h"
-#include "WorldPartition/HLOD/IWorldPartitionHLODUtilitiesModule.h"
-
-#include "Modules/ModuleManager.h"
 #endif
 
 DEFINE_LOG_CATEGORY_STATIC(LogHLODLayer, Log, All);
diff --git a/Engine/Source/Runtime/Engine/Private/WorldPartition/HLOD/HLODModifier.cpp b/Engine/Source/Runtime/Engine/Private/WorldPartition/HLOD/HLODModifier.cpp
new file mode 100644
--- /dev/null
+++ b/Engine/Source/Runtime/Engine/Private/WorldPartition/HLOD/HLODModifier.cpp
@@ -0,0 +1,9 @@
+// Copyright Epic Games, Inc. All Rights Reserved.
+
+#include "WorldPartition/HLOD/HLODModifier.h"
+
+
+UWorldPartitionHLODModifier::UWorldPartitionHLODModifier(const FObjectInitializer& ObjectInitializer)
+	: Super(ObjectInitializer)
+{
+}
diff --git a/Engine/Source/Runtime/Engine/Private/WorldPartition/HLOD/HLODSubsystem.cpp b/Engine/Source/Runtime/Engine/Private/WorldPartition/HLOD/HLODSubsystem.cpp
--- a/Engine/Source/Runtime/Engine/Private/WorldPartition/HLOD/HLODSubsystem.cpp
+++ b/Engine/Source/Runtime/Engine/Private/WorldPartition/HLOD/HLODSubsystem.cpp
@@ -68,6 +68,14 @@ static TAutoConsoleVariable<int32> CVarHLODWarmupVTSizeClamp(
 
 namespace FHLODSubsystem
 {
+	const UWorldPartitionRuntimeCell* GetActorRuntimeCell(AActor* InActor)
+	{
+		const ULevel* ActorLevel = InActor->GetLevel();
+		const ULevelStreaming* LevelStreaming = FLevelUtils::FindStreamingLevel(ActorLevel);
+		const UWorldPartitionLevelStreamingDynamic* LevelStreamingDynamic = Cast<UWorldPartitionLevelStreamingDynamic>(LevelStreaming);
+		return LevelStreamingDynamic ? LevelStreamingDynamic->GetWorldPartitionRuntimeCell() : nullptr;
+	}
+
 	UWorldPartition* GetWorldPartition(AWorldPartitionHLOD* InWorldPartitionHLOD)
 	{
 		// Alwaysloaded Cell level will have a WorldPartition
@@ -75,9 +83,9 @@ namespace FHLODSubsystem
 		{
 			return WorldPartition;
 		} // If not find it through the cell
-		else if (UWorldPartitionLevelStreamingDynamic* LevelStreaming = Cast<UWorldPartitionLevelStreamingDynamic>(FLevelUtils::FindStreamingLevel(InWorldPartitionHLOD->GetLevel())))
+		else if (const UWorldPartitionRuntimeCell* RuntimeCell = GetActorRuntimeCell(InWorldPartitionHLOD))
 		{
-			return LevelStreaming->GetWorldPartitionRuntimeCell()->GetCellOwner()->GetOuterWorld()->GetWorldPartition();;
+			return RuntimeCell->GetCellOwner()->GetOuterWorld()->GetWorldPartition();
 		}
 
 		return nullptr;
@@ -243,6 +251,25 @@ void UHLODSubsystem::UnregisterHLODActor(AWorldPartitionHLOD* InWorldPartitionHL
 	}
 }
 
+const TArray<AWorldPartitionHLOD*>& UHLODSubsystem::GetHLODActorsForCell(const UWorldPartitionRuntimeCell* InCell) const
+{
+	const UWorldPartition* WorldPartition = InCell->GetWorld()->GetWorldPartition();
+
+	const FWorldPartitionHLODRuntimeData* WorldPartitionHLODRuntimeData = WorldPartitionsHLODRuntimeData.Find(WorldPartition);
+	if (WorldPartitionHLODRuntimeData)
+	{
+		const FCellData* CellData = WorldPartitionHLODRuntimeData->CellsData.Find(InCell->GetFName());
+		if (CellData)
+		{
+			return CellData->LoadedHLODs;
+		}
+	}
+
+	// No HLOD found for the given cell, return a dummy array
+	static const TArray<AWorldPartitionHLOD*> DummyArray;
+	return DummyArray;
+}
+
 void UHLODSubsystem::OnCellShown(const UWorldPartitionRuntimeCell* InCell)
 {
 	const UWorldPartition* WorldPartition = InCell->GetCellOwner()->GetOuterWorld()->GetWorldPartition();
diff --git a/Engine/Source/Runtime/Engine/Private/WorldPartition/RuntimeSpatialHash/RuntimeSpatialHashHLOD.cpp b/Engine/Source/Runtime/Engine/Private/WorldPartition/RuntimeSpatialHash/RuntimeSpatialHashHLOD.cpp
--- a/Engine/Source/Runtime/Engine/Private/WorldPartition/RuntimeSpatialHash/RuntimeSpatialHashHLOD.cpp
+++ b/Engine/Source/Runtime/Engine/Private/WorldPartition/RuntimeSpatialHash/RuntimeSpatialHashHLOD.cpp
@@ -9,6 +9,7 @@
 #include "WorldPartition/HLOD/HLODLayer.h"
 #include "WorldPartition/HLOD/HLODActor.h"
 #include "WorldPartition/HLOD/HLODActorDesc.h"
+#include "WorldPartition/HLOD/HLODModifier.h"
 #include "WorldPartition/HLOD/IWorldPartitionHLODUtilitiesModule.h"
 #include "WorldPartition/DataLayer/DataLayerSubsystem.h"
 
@@ -319,9 +320,21 @@ static TMap<FName, FSpatialHashRuntimeGrid> CreateHLODGrids(TMap<UHLODLayer*, in
 			HLODGrid.LoadingRange = HLODLayer->GetLoadingRange();
 			HLODGrid.DebugColor = FLinearColor::Red;
 			HLODGrid.GridName = HLODLayer->GetRuntimeGrid(HLODLevel);
-			HLODGrid.bClientOnlyVisible = true;
 			HLODGrid.HLODLayer = HLODLayer;
 
+			// Temp solution to have replication working for hlod layers that have a modifier
+			// The real solution would probably be to check if we have any actors that are
+			// returning true for NeedsLoadForServer()
+			HLODGrid.bClientOnlyVisible = HLODLayer->GetHLODModifierClass() == nullptr;
+
+			// HLOD grid have a lower priority than standard grids
+			HLODGrid.Priority = 100 + HLODLevel;
+
+			// Setup grid debug color to match HLODColorationColors
+			const int32 LastLODColorationColorIdx = GEngine->HLODColorationColors.Num() - 1;
+			const int32 HLODColorIdx = FMath::Clamp(HLODLevel + 2, 0U, (int32)LastLODColorationColorIdx);
+			HLODGrid.DebugColor = GEngine->HLODColorationColors[HLODColorIdx];
+
 			HLODGrids.Emplace(HLODGrid.GridName, HLODGrid);
 		}
 	}
@@ -358,34 +371,39 @@ static void UpdateHLODGridsActors(UWorld* World, const TMap<FName, FSpatialHashR
 		}
 	}
 
-	// Create missing HLOD grid actors 
+	// Create missing HLOD grid actors or update existing actors
 	for (const auto& HLODGridEntry : HLODGrids)
 	{
 		FName GridName = HLODGridEntry.Key;
-		if (!ExistingGridActors.Contains(GridName))
-		{
-			const FSpatialHashRuntimeGrid& GridSettings = HLODGridEntry.Value;
+		const FSpatialHashRuntimeGrid& GridSettings = HLODGridEntry.Value;
 
+		bool bDirty = false;
+
+		ASpatialHashRuntimeGridInfo* GridActor = ExistingGridActors.FindRef(GridName);
+		if (!GridActor)
+		{
 			FActorSpawnParameters SpawnParams;
 			SpawnParams.bCreateActorPackage = true;
-			ASpatialHashRuntimeGridInfo* GridActor = World->SpawnActor<ASpatialHashRuntimeGridInfo>(SpawnParams);
+			GridActor = World->SpawnActor<ASpatialHashRuntimeGridInfo>(SpawnParams);
 			GridActor->Tags.Add(HLODGridTag);
-			GridActor->SetActorLabel(GridSettings.GridName.ToString());
-			GridActor->GridSettings = GridSettings;
-
-			// Setup grid debug color to match HLODColorationColors
-			const uint32 LastLODColorationColorIdx = GEngine->HLODColorationColors.Num() - 1;
-			uint32 HLODLevel;
-			if (!LexTryParseString(HLODLevel, *GridSettings.GridName.ToString().Mid(HLODGridTagLen, 1)))
-			{
-				HLODLevel = LastLODColorationColorIdx;
-			}
+			bDirty = true;
+		}
 
-			GridActor->GridSettings.Priority = 100 + HLODLevel;
+		const FString ActorLabel = GridSettings.GridName.ToString();
+		if (GridActor->GetActorLabel() != ActorLabel)
+		{
+			GridActor->SetActorLabel(ActorLabel);
+			bDirty = true;
+		}
 
-			HLODLevel = FMath::Clamp(HLODLevel + 2, 0U, LastLODColorationColorIdx);
-			GridActor->GridSettings.DebugColor = GEngine->HLODColorationColors[HLODLevel];
+		if (GridActor->GridSettings != GridSettings)
+		{
+			GridActor->GridSettings = GridSettings;
+			bDirty = true;
+		}			
 
+		if (bDirty)
+		{
 			SavePackage(GridActor->GetPackage(), SourceControlHelper);
 		}
 	}
diff --git a/Engine/Source/Runtime/Engine/Private/WorldPartition/WorldPartitionRuntimeSpatialHash.cpp b/Engine/Source/Runtime/Engine/Private/WorldPartition/WorldPartitionRuntimeSpatialHash.cpp
--- a/Engine/Source/Runtime/Engine/Private/WorldPartition/WorldPartitionRuntimeSpatialHash.cpp
+++ b/Engine/Source/Runtime/Engine/Private/WorldPartition/WorldPartitionRuntimeSpatialHash.cpp
@@ -919,6 +919,29 @@ void FSpatialHashStreamingGrid::DrawStreamingSource2D(UCanvas* Canvas, const FSp
 
 // ------------------------------------------------------------------------------------------------
 
+#if WITH_EDITORONLY_DATA
+
+bool FSpatialHashRuntimeGrid::operator == (const FSpatialHashRuntimeGrid& Other) const
+{
+	return GridName == Other.GridName &&
+		   CellSize == Other.CellSize &&
+		   LoadingRange == Other.LoadingRange &&
+		   bBlockOnSlowStreaming == Other.bBlockOnSlowStreaming &&
+		   Priority == Other.Priority &&
+		   DebugColor == Other.DebugColor &&
+		   bClientOnlyVisible == Other.bClientOnlyVisible &&
+		   HLODLayer == Other.HLODLayer;
+}
+
+bool FSpatialHashRuntimeGrid::operator != (const FSpatialHashRuntimeGrid& Other) const
+{
+	return !(*this == Other);
+}
+
+#endif
+
+// ------------------------------------------------------------------------------------------------
+
 ASpatialHashRuntimeGridInfo::ASpatialHashRuntimeGridInfo(const FObjectInitializer& ObjectInitializer)
 	: Super(ObjectInitializer)
 {
diff --git a/Engine/Source/Runtime/Engine/Public/WorldPartition/HLOD/HLODBuilder.h b/Engine/Source/Runtime/Engine/Public/WorldPartition/HLOD/HLODBuilder.h
--- a/Engine/Source/Runtime/Engine/Public/WorldPartition/HLOD/HLODBuilder.h
+++ b/Engine/Source/Runtime/Engine/Public/WorldPartition/HLOD/HLODBuilder.h
@@ -4,6 +4,7 @@
 
 #include "CoreMinimal.h"
 #include "Algo/Transform.h"
+#include "Templates/SubclassOf.h"
 
 #include "HLODBuilder.generated.h"
 
@@ -37,16 +38,22 @@ public:
 struct FHLODBuildContext
 {
 	/** World for which HLODs are being built */
-	UWorld*		World;
+	UWorld*	World;
+
+	/** Actors that will be represented by this HLOD */
+	TArray<AActor*> SourceActors;
 
 	/** Outer to use for generated assets */
-	UObject*	AssetsOuter;
+	UObject* AssetsOuter;
 
 	/** Base name to use for generated assets */
-	FString		AssetsBaseName;
+	FString	AssetsBaseName;
+
+	// Location of this HLOD actor in the world
+	FVector WorldPosition;
 
 	/** Minimum distance at which the HLOD is expected to be displayed */
-	double		MinVisibleDistance;
+	double MinVisibleDistance;
 };
 
 
@@ -70,7 +77,10 @@ public:
 	 * Build an HLOD representation of the input actors.
 	 * Components returned by this method needs to be properly outered & assigned to your target (HLOD) actor.
 	 */
-	TArray<UActorComponent*> Build(const FHLODBuildContext& InHLODBuildContext, const TArray<AActor*>& InSourceActors) const;
+	TArray<UActorComponent*> Build(const FHLODBuildContext& InHLODBuildContext) const;
+
+	UE_DEPRECATED(5.2, "Use Build() method that takes a single FHLODBuildContext parameter. SourceActors are now part of the HLOD build context object.")
+	TArray<UActorComponent*> Build(const FHLODBuildContext& InHLODBuildContext, const TArray<AActor*>& InSourceActors) const { return Build(InHLODBuildContext); }
 
 	/**
 	 * Return the setting subclass associated with this builder.
@@ -101,7 +111,7 @@ public:
 	/**
 	 * Components created with this method need to be properly outered & assigned to your target actor.
 	 */
-	virtual TArray<UActorComponent*> Build(const FHLODBuildContext& InHLODBuildContext, const TArray<UActorComponent*>& InSourceComponents) const PURE_VIRTUAL(UHLODBuilder::CreateComponents, return {};);
+	virtual TArray<UActorComponent*> Build(const FHLODBuildContext& InHLODBuildContext, const TArray<UActorComponent*>& InSourceComponents) const PURE_VIRTUAL(UHLODBuilder::Build, return {};);
 
 	/**
 	 * From a set of actors, compute a unique hash from their properties that are relevant to HLOD generation.
diff --git a/Engine/Source/Runtime/Engine/Public/WorldPartition/HLOD/HLODLayer.h b/Engine/Source/Runtime/Engine/Public/WorldPartition/HLOD/HLODLayer.h
--- a/Engine/Source/Runtime/Engine/Public/WorldPartition/HLOD/HLODLayer.h
+++ b/Engine/Source/Runtime/Engine/Public/WorldPartition/HLOD/HLODLayer.h
@@ -11,13 +11,14 @@
 #include "EngineDefines.h"
 #include "Engine/MeshMerging.h"
 
-#include "WorldPartition/HLOD/HLODBuilder.h"
-
 #include "HLODLayer.generated.h"
 
 class AActor;
+class UHLODBuilder;
+class UHLODBuilderSettings;
 class UMaterial;
 class UWorldPartition;
+class UWorldPartitionHLODModifier;
 class FWorldPartitionActorDesc;
 class FWorldPartitionActorDescView;
 
@@ -52,6 +53,7 @@ public:
 	void SetLayerType(EHLODLayerType InLayerType) { LayerType = InLayerType; }
 	const TSubclassOf<UHLODBuilder> GetHLODBuilderClass() const { return HLODBuilderClass; }
 	const UHLODBuilderSettings* GetHLODBuilderSettings() const { return HLODBuilderSettings; }
+	const TSubclassOf<UWorldPartitionHLODModifier> GetHLODModifierClass() const { return HLODModifierClass; }
 	FName GetRuntimeGrid(uint32 InHLODLevel) const;
 	int32 GetCellSize() const { return !bIsSpatiallyLoaded ? 0 : CellSize; }
 	double GetLoadingRange() const { return !bIsSpatiallyLoaded ? WORLD_MAX : LoadingRange; }
@@ -76,7 +78,7 @@ private:
 	UPROPERTY(EditAnywhere, Config, Category=HLOD)
 	EHLODLayerType LayerType;
 
-	/** HLODBuilder class */
+	/** HLOD Builder class */
 	UPROPERTY(EditAnywhere, Config, Category=HLOD, meta = (DisplayName = "HLOD Builder Class", EditConditionHides, EditCondition = "LayerType == EHLODLayerType::Custom"))
 	TSubclassOf<UHLODBuilder> HLODBuilderClass;
 
@@ -99,6 +101,10 @@ private:
 	UPROPERTY(EditAnywhere, Config, Category=HLOD, meta = (EditConditionHides, EditCondition = "bIsSpatiallyLoaded"))
 	TSoftObjectPtr<UHLODLayer> ParentLayer;
 
+	/** HLOD Modifier class, to allow changes to the HLOD at runtime */
+	UPROPERTY(EditAnywhere, Config, Category = HLOD, AdvancedDisplay, meta = (DisplayName = "HLOD Modifier Class"))
+	TSubclassOf<UWorldPartitionHLODModifier> HLODModifierClass;
+
 private:
 	friend class FWorldPartitionHLODUtilities;
 
diff --git a/Engine/Source/Runtime/Engine/Public/WorldPartition/HLOD/HLODModifier.h b/Engine/Source/Runtime/Engine/Public/WorldPartition/HLOD/HLODModifier.h
new file mode 100644
--- /dev/null
+++ b/Engine/Source/Runtime/Engine/Public/WorldPartition/HLOD/HLODModifier.h
@@ -0,0 +1,28 @@
+// Copyright Epic Games, Inc. All Rights Reserved.
+
+#pragma once
+
+#include "CoreMinimal.h"
+#include "Templates/SubclassOf.h"
+
+#include "HLODModifier.generated.h"
+
+
+class UActorComponent;
+class UHLODBuilder;
+struct FHLODBuildContext;
+
+/**
+ * Base class for all runtime HLOD modifiers
+ */
+UCLASS(Abstract)
+class ENGINE_API UWorldPartitionHLODModifier : public UObject
+{
+	 GENERATED_UCLASS_BODY()
+
+public:
+	virtual bool CanModifyHLOD(TSubclassOf<UHLODBuilder> InHLODBuilderClass) const PURE_VIRTUAL(UWorldPartitionHLODModifier::CanModifyHLOD, return false;);
+	
+	virtual void BeginHLODBuild(const FHLODBuildContext& InHLODBuildContext) PURE_VIRTUAL(UWorldPartitionHLODModifier::BeginHLODBuild, );
+	virtual void EndHLODBuild(TArray<UActorComponent*>& InOutComponents) PURE_VIRTUAL(UWorldPartitionHLODModifier::EndHLODBuild, );
+};
diff --git a/Engine/Source/Runtime/Engine/Public/WorldPartition/HLOD/HLODSubsystem.h b/Engine/Source/Runtime/Engine/Public/WorldPartition/HLOD/HLODSubsystem.h
--- a/Engine/Source/Runtime/Engine/Public/WorldPartition/HLOD/HLODSubsystem.h
+++ b/Engine/Source/Runtime/Engine/Public/WorldPartition/HLOD/HLODSubsystem.h
@@ -56,16 +56,17 @@ public:
 
 	bool RequestUnloading(const UWorldPartitionRuntimeCell* InCell);
 
+	const TArray<AWorldPartitionHLOD*>& GetHLODActorsForCell(const UWorldPartitionRuntimeCell* InCell) const;
+
 	static bool IsHLODEnabled();
 	
 private:
 	struct FCellData
 	{
-		bool						bIsCellVisible;
-		TSet<AWorldPartitionHLOD*>	LoadedHLODs;
-
-		uint32						WarmupStartFrame;
-		uint32						WarmupEndFrame;
+		bool bIsCellVisible;
+		TArray<AWorldPartitionHLOD*> LoadedHLODs;
+		uint32 WarmupStartFrame;
+		uint32 WarmupEndFrame;
 
 		FCellData() 
 			: bIsCellVisible(false)
diff --git a/Engine/Source/Runtime/Engine/Public/WorldPartition/WorldPartitionRuntimeSpatialHash.h b/Engine/Source/Runtime/Engine/Public/WorldPartition/WorldPartitionRuntimeSpatialHash.h
--- a/Engine/Source/Runtime/Engine/Public/WorldPartition/WorldPartitionRuntimeSpatialHash.h
+++ b/Engine/Source/Runtime/Engine/Public/WorldPartition/WorldPartitionRuntimeSpatialHash.h
@@ -210,6 +210,9 @@ struct FSpatialHashRuntimeGrid
 
 	UPROPERTY()
 	TObjectPtr<const UHLODLayer> HLODLayer;
+
+	bool operator == (const FSpatialHashRuntimeGrid& Other) const;
+	bool operator != (const FSpatialHashRuntimeGrid& Other) const;
 #endif
 };
 
