WorldPartition - HLOD: Prevent UObject reuse during HLOD build

diff --git a/Engine/Plugins/Editor/WorldPartitionHLODUtilities/Source/Private/WorldPartitionHLODUtilities.cpp b/Engine/Plugins/Editor/WorldPartitionHLODUtilities/Source/Private/WorldPartitionHLODUtilities.cpp
--- a/Engine/Plugins/Editor/WorldPartitionHLODUtilities/Source/Private/WorldPartitionHLODUtilities.cpp
+++ b/Engine/Plugins/Editor/WorldPartitionHLODUtilities/Source/Private/WorldPartitionHLODUtilities.cpp
@@ -34,6 +34,7 @@
 #include "StaticMeshCompiler.h"
 #include "Templates/UniquePtr.h"
 #include "TextureCompiler.h"
+#include "UObject/MetaData.h"
 #include "UObject/GCObjectScopeGuard.h"
 
 
@@ -574,7 +575,29 @@ uint32 FWorldPartitionHLODUtilities::BuildHLOD(AWorldPartitionHLOD* InHLODActor)
 
 	// Clear stats as we're about to refresh them
 	InHLODActor->ResetStats();
-	
+
+	// Rename previous assets found in the HLOD actor package.
+	// Move the previous asset(s) to the transient package, to avoid any object reuse during the build
+	{
+	    TArray<UObject*> ObjectsToRename;
+	    ForEachObjectWithOuter(InHLODActor->GetPackage(), [&ObjectsToRename](UObject* Obj)
+	    {
+		    if (!Obj->IsA<AActor>() && !Obj->IsA<UMetaData>())
+		    {
+			    ObjectsToRename.Add(Obj);
+		    }
+	    }, false);
+    
+	    
+	    for (UObject* Obj : ObjectsToRename)
+	    {
+		    // Make sure the old object is not used by anything
+		    Obj->ClearFlags(RF_Standalone | RF_Public);
+		    const FName OldRenamed = MakeUniqueObjectName(GetTransientPackage(), Obj->GetClass(), *FString::Printf(TEXT("OLD_%s"), *Obj->GetName()));
+		    Obj->Rename(*OldRenamed.ToString(), GetTransientPackage(), REN_DontCreateRedirectors | REN_NonTransactional | REN_ForceNoResetLoaders);
+	    }
+	}
+
 	// Gather stats from the input to our HLOD build
 	GatherInputStats(InHLODActor, LevelStreaming->GetLoadedLevel()->Actors);
 	
