World Partition - HLOD: Fix false error for non spatially loaded actors included in HLOD setup
* Must check the actor set instance bIsSpatiallyLoaded instead

diff --git a/Engine/Plugins/Editor/WorldPartitionHLODUtilities/Source/Private/WorldPartitionHLODUtilities.cpp b/Engine/Plugins/Editor/WorldPartitionHLODUtilities/Source/Private/WorldPartitionHLODUtilities.cpp
--- a/Engine/Plugins/Editor/WorldPartitionHLODUtilities/Source/Private/WorldPartitionHLODUtilities.cpp
+++ b/Engine/Plugins/Editor/WorldPartitionHLODUtilities/Source/Private/WorldPartitionHLODUtilities.cpp
@@ -117,14 +117,14 @@ TArray<AWorldPartitionHLOD*> FWorldPartitionHLODUtilities::CreateHLODActors(FHLO
 	{
 		const FWorldPartitionActorDescView& ActorDescView = ActorInstance.GetActorDescView();
 
-		if (!ActorDescView.GetIsSpatiallyLoaded())
-		{
-			UE_LOG(LogHLODBuilder, Error, TEXT("Tried to included non-spatially loaded actor %s into HLOD"), *ActorDescView.GetActorName().ToString());
-			continue;
-		}
-
 		if (ActorDescView.GetActorIsHLODRelevant())
 		{
+			if (!ActorInstance.ActorSetInstance->bIsSpatiallyLoaded)
+			{
+				UE_LOG(LogHLODBuilder, Warning, TEXT("Tried to included non-spatially loaded actor %s into HLOD"), *ActorDescView.GetActorName().ToString());
+				continue;
+			}
+
 			UHLODLayer* HLODLayer = UHLODLayer::GetHLODLayer(ActorDescView, InCreationParams.WorldPartition);
 			if (HLODLayer)
 			{
