World Partition - HLOD: Added a base CRC key to the mesh merge builder, to force a rebuild of HLODs generated from this builder when needed

diff --git a/Engine/Plugins/Editor/WorldPartitionHLODUtilities/Source/Private/HLODBuilderMeshMerge.cpp b/Engine/Plugins/Editor/WorldPartitionHLODUtilities/Source/Private/HLODBuilderMeshMerge.cpp
--- a/Engine/Plugins/Editor/WorldPartitionHLODUtilities/Source/Private/HLODBuilderMeshMerge.cpp
+++ b/Engine/Plugins/Editor/WorldPartitionHLODUtilities/Source/Private/HLODBuilderMeshMerge.cpp
@@ -44,6 +44,10 @@ uint32 UHLODBuilderMeshMergeSettings::GetCRC() const
 
 	FArchiveCrc32 Ar;
 
+	// Base mesh merge key, changing this will force a rebuild of all HLODs from this builder
+	FString HLODBaseKey = "89D89284DD3847FA90C5998E06DD8FEC";
+	Ar << HLODBaseKey;
+
 	Ar << This.MeshMergeSettings;
 	UE_LOG(LogHLODBuilder, VeryVerbose, TEXT(" - MeshMergeSettings = %d"), Ar.GetCrc());
 
