Ensure actors references are loaded when building HLODs.
* Specifically, parent actors could be missing, causing the HLOD source actors transform to be wrong

diff --git a/Engine/Plugins/Editor/WorldPartitionHLODUtilities/Source/Private/WorldPartitionHLODUtilities.cpp b/Engine/Plugins/Editor/WorldPartitionHLODUtilities/Source/Private/WorldPartitionHLODUtilities.cpp
--- a/Engine/Plugins/Editor/WorldPartitionHLODUtilities/Source/Private/WorldPartitionHLODUtilities.cpp
+++ b/Engine/Plugins/Editor/WorldPartitionHLODUtilities/Source/Private/WorldPartitionHLODUtilities.cpp
@@ -91,6 +91,18 @@ static uint32 GetCRC(const UHLODLayer* InHLODLayer)
 	return CRC;
 }
 
+static uint32 GetCRC(const TArray<FHLODSubActor>& InSubActors)
+{
+	uint32 CRC = 0;
+
+	for (const FHLODSubActor& HLODSubActor : InSubActors)
+	{
+		CRC = HashCombine(GetTypeHash(HLODSubActor), CRC);
+	}
+
+	return CRC;
+}
+
 static uint32 ComputeHLODHash(AWorldPartitionHLOD* InHLODActor, const TArray<AActor*>& InActors)
 {
 	FArchiveCrc32 Ar;
@@ -99,6 +111,11 @@ static uint32 ComputeHLODHash(AWorldPartitionHLOD* InHLODActor, const TArray<AAc
 	FString HLODBaseKey = "A6C1517FC49B478E833A540E2C06289F";
 	Ar << HLODBaseKey;
 
+	// Sub Actors
+	uint32 SubActorsHash = GetCRC(InHLODActor->GetSubActors());
+	UE_LOG(LogHLODBuilder, VeryVerbose, TEXT(" - Sub Actors (%d actors) = %x"), InHLODActor->GetSubActors().Num(), SubActorsHash);
+	Ar << SubActorsHash;
+
 	// HLOD Layer
 	uint32 HLODLayerHash = GetCRC(InHLODActor->GetSubActorsHLODLayer());
 	UE_LOG(LogHLODBuilder, VeryVerbose, TEXT(" - HLOD Layer (%s) = %x"), *InHLODActor->GetSubActorsHLODLayer()->GetName(), HLODLayerHash);
@@ -117,9 +134,33 @@ static uint32 ComputeHLODHash(AWorldPartitionHLOD* InHLODActor, const TArray<AAc
 	return Ar.GetCrc();
 }
 
+void AddSubActor(const FWorldPartitionActorDescView& ActorDescView, const IStreamingGenerationContext::FActorInstance& ActorInstance, TSet<FHLODSubActor>& SubActors)
+{
+PRAGMA_DISABLE_DEPRECATION_WARNINGS
+	// Leaving this as deprecated for now until we fix up the serialization format for SubActorsInfo and do proper upgrade/deprecation 
+	FName ActorPath = ActorDescView.GetActorPath();
+PRAGMA_ENABLE_DEPRECATION_WARNINGS
+
+	// Add the actor
+	bool bIsAlreadyInSet = false;
+	FHLODSubActor SubActor(ActorDescView.GetGuid(), ActorDescView.GetActorPackage(), ActorPath, ActorInstance.GetContainerID(), ActorInstance.GetActorDescContainer()->GetContainerPackage(), ActorInstance.GetTransform());
+	SubActors.Add(SubActor, &bIsAlreadyInSet);
+
+	if (!bIsAlreadyInSet)
+	{
+		// Add its references
+		const FActorDescViewMap* ActorDescViewMap = ActorInstance.ActorSetInstance->ContainerInstance->ActorDescViewMap;
+		for (const FGuid& ReferenceGuid : ActorDescView.GetReferences())
+		{
+			const FWorldPartitionActorDescView& RefActorDescView = ActorDescViewMap->FindByGuidChecked(ReferenceGuid);
+			AddSubActor(RefActorDescView, ActorInstance, SubActors);
+		}
+	}
+}
+
 TArray<AWorldPartitionHLOD*> FWorldPartitionHLODUtilities::CreateHLODActors(FHLODCreationContext& InCreationContext, const FHLODCreationParams& InCreationParams, const TArray<IStreamingGenerationContext::FActorInstance>& InActors, const TArray<const UDataLayerInstance*>& InDataLayersInstances)
 {
-	TMap<UHLODLayer*, TArray<FHLODSubActor>> SubActorsPerHLODLayer;
+	TMap<UHLODLayer*, TSet<FHLODSubActor>> SubActorsPerHLODLayer;
 
 	for (const IStreamingGenerationContext::FActorInstance& ActorInstance : InActors)
 	{
@@ -136,12 +177,8 @@ TArray<AWorldPartitionHLOD*> FWorldPartitionHLODUtilities::CreateHLODActors(FHLO
 			UHLODLayer* HLODLayer = UHLODLayer::GetHLODLayer(ActorDescView, InCreationParams.WorldPartition);
 			if (HLODLayer)
 			{
-				TArray<FHLODSubActor>& SubActors = SubActorsPerHLODLayer.FindOrAdd(HLODLayer);
-
-				// Leaving this as deprecated for now until we fix up the serialization format for SubActorsInfo and do proper upgrade/deprecation 
-PRAGMA_DISABLE_DEPRECATION_WARNINGS
-				SubActors.Emplace(ActorDescView.GetGuid(), ActorDescView.GetActorPackage(), ActorDescView.GetActorPath(), ActorInstance.GetContainerID(), ActorInstance.GetActorDescContainer()->GetContainerPackage(), ActorInstance.GetTransform());
-PRAGMA_ENABLE_DEPRECATION_WARNINGS
+				TSet<FHLODSubActor>& SubActors = SubActorsPerHLODLayer.FindOrAdd(HLODLayer);
+				AddSubActor(ActorDescView, ActorInstance, SubActors);
 			}
 		}
 	}
@@ -150,7 +187,7 @@ PRAGMA_ENABLE_DEPRECATION_WARNINGS
 	for (const auto& Pair : SubActorsPerHLODLayer)
 	{
 		const UHLODLayer* HLODLayer = Pair.Key;
-		const TArray<FHLODSubActor>& SubActors = Pair.Value;
+		const TSet<FHLODSubActor>& SubActors = Pair.Value;
 		check(!SubActors.IsEmpty());
 
 		auto ComputeCellHash = [](const FString& HLODLayerName, const FName CellName)
@@ -204,7 +241,7 @@ PRAGMA_ENABLE_DEPRECATION_WARNINGS
 			if (!bSubActorsChanged)
 			{
 				TArray<FHLODSubActor> A = HLODActor->GetSubActors();
-				TArray<FHLODSubActor> B = SubActors;
+				TArray<FHLODSubActor> B = SubActors.Array();
 				A.Sort();
 				B.Sort();
 				bSubActorsChanged = A != B;
@@ -212,7 +249,7 @@ PRAGMA_ENABLE_DEPRECATION_WARNINGS
 
 			if (bSubActorsChanged)
 			{
-				HLODActor->SetSubActors(SubActors);
+				HLODActor->SetSubActors(SubActors.Array());
 				bIsDirty = true;
 			}
 		}
diff --git a/Engine/Source/Runtime/Engine/Public/WorldPartition/HLOD/HLODSubActor.h b/Engine/Source/Runtime/Engine/Public/WorldPartition/HLOD/HLODSubActor.h
--- a/Engine/Source/Runtime/Engine/Public/WorldPartition/HLOD/HLODSubActor.h
+++ b/Engine/Source/Runtime/Engine/Public/WorldPartition/HLOD/HLODSubActor.h
@@ -96,6 +96,11 @@ struct FHLODSubActor
 
 		return Lhs.ActorGuid < Rhs.ActorGuid;
 	}
+
+	friend uint32 GetTypeHash(const FHLODSubActor& Key)
+	{
+		return HashCombine(GetTypeHash(Key.ActorGuid), GetTypeHash(Key.ContainerID));
+	}
 		
 	UPROPERTY()
 	FGuid ActorGuid;
