[Engine/Plugins/Editor]

diff --git a/Engine/Plugins/Editor/SequencerAnimTools/Source/SequencerAnimTools/Private/TrajectoryDrawInfo.cpp b/Engine/Plugins/Editor/SequencerAnimTools/Source/SequencerAnimTools/Private/TrajectoryDrawInfo.cpp
--- a/Engine/Plugins/Editor/SequencerAnimTools/Source/SequencerAnimTools/Private/TrajectoryDrawInfo.cpp
+++ b/Engine/Plugins/Editor/SequencerAnimTools/Source/SequencerAnimTools/Private/TrajectoryDrawInfo.cpp
@@ -1,7 +1,7 @@
 // Copyright Epic Games, Inc. All Rights Reserved.
 
 #include "TrajectoryDrawInfo.h"
-#include "TrailHierarchy.h"
+#include "SceneView.h"
 #include "TrajectoryCache.h"
 
 namespace UE
diff --git a/Engine/Plugins/Editor/StylusInput/Source/StylusInput/Private/SStylusInputDebugWidget.cpp b/Engine/Plugins/Editor/StylusInput/Source/StylusInput/Private/SStylusInputDebugWidget.cpp
--- a/Engine/Plugins/Editor/StylusInput/Source/StylusInput/Private/SStylusInputDebugWidget.cpp
+++ b/Engine/Plugins/Editor/StylusInput/Source/StylusInput/Private/SStylusInputDebugWidget.cpp
@@ -2,8 +2,8 @@
 
 #include "SStylusInputDebugWidget.h"
 
-#include "IStylusState.h"
 
+#include "IStylusInputModule.h"
 #include "Widgets/SBoxPanel.h"
 #include "Widgets/Text/STextBlock.h"
 #include "Widgets/Input/SCheckBox.h"
@@ -221,4 +221,4 @@ ECheckBoxState SStylusInputDebugWidget::IsInverted() const
 }
 
 
-#undef LOCTEXT_NAMESPACE
\ No newline at end of file
+#undef LOCTEXT_NAMESPACE
diff --git a/Engine/Plugins/Editor/StylusInput/Source/StylusInput/Private/SStylusInputDebugWidget.h b/Engine/Plugins/Editor/StylusInput/Source/StylusInput/Private/SStylusInputDebugWidget.h
--- a/Engine/Plugins/Editor/StylusInput/Source/StylusInput/Private/SStylusInputDebugWidget.h
+++ b/Engine/Plugins/Editor/StylusInput/Source/StylusInput/Private/SStylusInputDebugWidget.h
@@ -4,9 +4,11 @@
 
 #include "Widgets/SCompoundWidget.h"
 
-#include "IStylusInputModule.h"
 #include "IStylusState.h"
 
+class UStylusInputSubsystem;
+enum class ECheckBoxState : uint8;
+
 class SStylusInputDebugWidget : public SCompoundWidget, 
 	public IStylusMessageHandler
 {
@@ -47,4 +49,4 @@ private:
 
 	static FText GetVector2Text(FVector2D Value);
 	static FText GetFloatText(float Value);
-};
\ No newline at end of file
+};
diff --git a/Engine/Plugins/Editor/StylusInput/Source/StylusInput/Private/StylusInputModule.cpp b/Engine/Plugins/Editor/StylusInput/Source/StylusInput/Private/StylusInputModule.cpp
--- a/Engine/Plugins/Editor/StylusInput/Source/StylusInput/Private/StylusInputModule.cpp
+++ b/Engine/Plugins/Editor/StylusInput/Source/StylusInput/Private/StylusInputModule.cpp
@@ -1,14 +1,9 @@
 // Copyright Epic Games, Inc. All Rights Reserved.
 
 #include "IStylusInputModule.h"
-#include "CoreMinimal.h"
 
-#include "Framework/Docking/TabManager.h"
-#include "Interfaces/IMainFrameModule.h"
-#include "Logging/LogMacros.h"
 #include "Misc/App.h"
 #include "Widgets/Docking/SDockTab.h"
-#include "Widgets/SWindow.h"
 #include "WorkspaceMenuStructure.h"
 #include "WorkspaceMenuStructureModule.h"
 #include "SStylusInputDebugWidget.h"
@@ -139,4 +134,4 @@ TSharedRef<SDockTab> UStylusInputSubsystem::OnSpawnPluginTab(const FSpawnTabArgs
 
 
 
-#undef LOCTEXT_NAMESPACE
\ No newline at end of file
+#undef LOCTEXT_NAMESPACE
diff --git a/Engine/Plugins/Editor/StylusInput/Source/StylusInput/Private/WindowsRealTimeStylusPlugin.h b/Engine/Plugins/Editor/StylusInput/Source/StylusInput/Private/WindowsRealTimeStylusPlugin.h
--- a/Engine/Plugins/Editor/StylusInput/Source/StylusInput/Private/WindowsRealTimeStylusPlugin.h
+++ b/Engine/Plugins/Editor/StylusInput/Source/StylusInput/Private/WindowsRealTimeStylusPlugin.h
@@ -2,7 +2,7 @@
 
 #pragma once
 
-#include "CoreMinimal.h"
+#include "HAL/Platform.h"
 
 #if PLATFORM_WINDOWS
 #include "Windows/WindowsHWrapper.h"
@@ -175,4 +175,4 @@ private:
 	void RemoveTabletContext(IRealTimeStylus* InRealTimeStylus, TABLET_CONTEXT_ID TabletID);
 };
 
-#endif // PLATFORM_WINDOWS
\ No newline at end of file
+#endif // PLATFORM_WINDOWS
diff --git a/Engine/Plugins/Editor/WorldPartitionHLODUtilities/Source/Private/HLODBuilderMeshApproximate.cpp b/Engine/Plugins/Editor/WorldPartitionHLODUtilities/Source/Private/HLODBuilderMeshApproximate.cpp
--- a/Engine/Plugins/Editor/WorldPartitionHLODUtilities/Source/Private/HLODBuilderMeshApproximate.cpp
+++ b/Engine/Plugins/Editor/WorldPartitionHLODUtilities/Source/Private/HLODBuilderMeshApproximate.cpp
@@ -3,16 +3,15 @@
 #include "HLODBuilderMeshApproximate.h"
 
 #include "Algo/ForEach.h"
-#include "Algo/Transform.h"
+#include "Components/StaticMeshComponent.h"
 #include "UObject/Package.h"
-#include "WorldPartition/HLOD/HLODActor.h"
-#include "WorldPartition/HLOD/HLODLayer.h"
 
+#include "Concepts/StaticStructProvider.h"
 #include "Engine/Engine.h"
 #include "Engine/StaticMesh.h"
+#include "Engine/Texture2D.h"
 #include "Materials/Material.h"
 #include "Materials/MaterialInstance.h"
-#include "MaterialUtilities.h"
 
 #include "Modules/ModuleManager.h"
 #include "IGeometryProcessingInterfacesModule.h"
diff --git a/Engine/Plugins/Editor/WorldPartitionHLODUtilities/Source/Private/HLODBuilderMeshMerge.cpp b/Engine/Plugins/Editor/WorldPartitionHLODUtilities/Source/Private/HLODBuilderMeshMerge.cpp
--- a/Engine/Plugins/Editor/WorldPartitionHLODUtilities/Source/Private/HLODBuilderMeshMerge.cpp
+++ b/Engine/Plugins/Editor/WorldPartitionHLODUtilities/Source/Private/HLODBuilderMeshMerge.cpp
@@ -3,13 +3,12 @@
 #include "HLODBuilderMeshMerge.h"
 
 #include "Algo/ForEach.h"
+#include "Concepts/StaticStructProvider.h"
 #include "Engine/Engine.h"
 #include "Engine/StaticMesh.h"
 #include "Materials/Material.h"
 #include "Components/StaticMeshComponent.h"
 
-#include "WorldPartition/HLOD/HLODActor.h"
-#include "WorldPartition/HLOD/HLODLayer.h"
 
 #include "IMeshMergeUtilities.h"
 #include "MeshMergeModule.h"
diff --git a/Engine/Plugins/Editor/WorldPartitionHLODUtilities/Source/Private/HLODBuilderMeshSimplify.cpp b/Engine/Plugins/Editor/WorldPartitionHLODUtilities/Source/Private/HLODBuilderMeshSimplify.cpp
--- a/Engine/Plugins/Editor/WorldPartitionHLODUtilities/Source/Private/HLODBuilderMeshSimplify.cpp
+++ b/Engine/Plugins/Editor/WorldPartitionHLODUtilities/Source/Private/HLODBuilderMeshSimplify.cpp
@@ -3,16 +3,15 @@
 #include "HLODBuilderMeshSimplify.h"
 
 #include "Algo/ForEach.h"
-#include "Algo/Transform.h"
+#include "Concepts/StaticStructProvider.h"
 #include "Engine/Engine.h"
 #include "Engine/StaticMesh.h"
 #include "Materials/Material.h"
 #include "Components/StaticMeshComponent.h"
 
-#include "WorldPartition/HLOD/HLODActor.h"
-#include "WorldPartition/HLOD/HLODLayer.h"
 
 #include "IMeshMergeUtilities.h"
+#include "Math/PerspectiveMatrix.h"
 #include "MeshMergeModule.h"
 #include "Modules/ModuleManager.h"
 
diff --git a/Engine/Plugins/Editor/WorldPartitionHLODUtilities/Source/Private/WorldPartition/HLOD/Modifiers/HLODModifierMeshDestruction.cpp b/Engine/Plugins/Editor/WorldPartitionHLODUtilities/Source/Private/WorldPartition/HLOD/Modifiers/HLODModifierMeshDestruction.cpp
--- a/Engine/Plugins/Editor/WorldPartitionHLODUtilities/Source/Private/WorldPartition/HLOD/Modifiers/HLODModifierMeshDestruction.cpp
+++ b/Engine/Plugins/Editor/WorldPartitionHLODUtilities/Source/Private/WorldPartition/HLOD/Modifiers/HLODModifierMeshDestruction.cpp
@@ -2,13 +2,14 @@
 
 #include "WorldPartition/HLOD/Modifiers/HLODModifierMeshDestruction.h"
 
-#include "Algo/Transform.h"
+#include "Components/StaticMeshComponent.h"
 #include "IMeshMergeExtension.h"
+#include "IMeshMergeUtilities.h"
 #include "Materials/MaterialInstanceConstant.h"
-#include "Materials/MaterialInterface.h"
 #include "MeshMergeDataTracker.h"
 #include "MeshMergeModule.h"
 #include "Modules/ModuleManager.h"
+#include "StaticMeshAttributes.h"
 #include "WorldPartition/HLOD/HLODBuilder.h"
 #include "WorldPartition/HLOD/HLODDestruction.h"
 #include "WorldPartition/HLOD/DestructibleHLODComponent.h"
@@ -170,4 +171,4 @@ void UWorldPartitionHLODModifierMeshDestruction::EndHLODBuild(TArray<UActorCompo
 	}
 
 	delete DestructionMeshMergeExtension;
-}
\ No newline at end of file
+}
diff --git a/Engine/Plugins/Editor/WorldPartitionHLODUtilities/Source/Private/WorldPartitionHLODUtilities.cpp b/Engine/Plugins/Editor/WorldPartitionHLODUtilities/Source/Private/WorldPartitionHLODUtilities.cpp
--- a/Engine/Plugins/Editor/WorldPartitionHLODUtilities/Source/Private/WorldPartitionHLODUtilities.cpp
+++ b/Engine/Plugins/Editor/WorldPartitionHLODUtilities/Source/Private/WorldPartitionHLODUtilities.cpp
@@ -1,19 +1,23 @@
 // Copyright Epic Games, Inc. All Rights Reserved.
 
 #include "WorldPartitionHLODUtilities.h"
+#include "WorldPartition/WorldPartitionActorDescView.h"
+#include "WorldPartition/WorldPartitionStreamingGeneration.h"
 
 #if WITH_EDITOR
 
-#include "WorldPartition/ActorDescContainer.h"
+#include "BodySetupEnums.h"
 #include "WorldPartition/WorldPartition.h"
-#include "WorldPartition/WorldPartitionHandle.h"
+#include "Components/InstancedStaticMeshComponent.h"
 #include "WorldPartition/DataLayer/DataLayerInstance.h"
+#include "Engine/CollisionProfile.h"
 #include "WorldPartition/HLOD/HLODActor.h"
+#include "Engine/Texture.h"
 #include "WorldPartition/HLOD/HLODSubActor.h"
-#include "WorldPartition/HLOD/HLODActorDesc.h"
-#include "WorldPartition/HLOD/HLODBuilder.h"
+#include "StaticMeshResources.h"
 #include "WorldPartition/HLOD/HLODLayer.h"
 #include "WorldPartition/HLOD/HLODModifier.h"
+#include "WorldPartition/HLOD/HLODStats.h"
 #include "WorldPartition/WorldPartitionLevelStreamingDynamic.h"
 
 #include "HLODBuilderInstancing.h"
@@ -21,18 +25,13 @@
 #include "HLODBuilderMeshSimplify.h"
 #include "HLODBuilderMeshApproximate.h"
 
-#include "Algo/Transform.h"
-#include "Components/StaticMeshComponent.h"
-#include "Engine/InstancedStaticMesh.h"
 #include "Engine/Level.h"
 #include "Engine/StaticMesh.h"
-#include "HAL/FileManager.h"
 #include "Materials/MaterialInstance.h"
 #include "PhysicsEngine/BodySetup.h"
 #include "ProfilingDebugging/ScopedTimers.h"
 #include "Serialization/ArchiveCrc32.h"
 #include "StaticMeshCompiler.h"
-#include "Templates/UniquePtr.h"
 #include "TextureCompiler.h"
 #include "UObject/MetaData.h"
 #include "UObject/GCObjectScopeGuard.h"
