Fix typo in comments.

diff --git a/Engine/Plugins/Editor/WorldPartitionHLODUtilities/Source/Private/HLODBuilderMeshApproximate.cpp b/Engine/Plugins/Editor/WorldPartitionHLODUtilities/Source/Private/HLODBuilderMeshApproximate.cpp
--- a/Engine/Plugins/Editor/WorldPartitionHLODUtilities/Source/Private/HLODBuilderMeshApproximate.cpp
+++ b/Engine/Plugins/Editor/WorldPartitionHLODUtilities/Source/Private/HLODBuilderMeshApproximate.cpp
@@ -89,7 +89,7 @@ TArray<UActorComponent*> UHLODBuilderMeshApproximate::Build(const FHLODBuildCont
 	FMeshApproximationSettings UseSettings = MeshApproximateSettings->MeshApproximationSettings; // Make a copy as we may tweak some values
 	UMaterialInterface* HLODMaterial = MeshApproximateSettings->HLODMaterial.LoadSynchronous();
 
-	// When using automatic textuse sizing based on draw distance, use the MinVisibleDistance for this HLOD.
+	// When using automatic texture sizing based on draw distance, use the MinVisibleDistance for this HLOD.
 	if (UseSettings.MaterialSettings.TextureSizingType == TextureSizingType_AutomaticFromMeshDrawDistance)
 	{
 		UseSettings.MaterialSettings.MeshMinDrawDistance = InHLODBuildContext.MinVisibleDistance;
diff --git a/Engine/Plugins/Editor/WorldPartitionHLODUtilities/Source/Private/HLODBuilderMeshMerge.cpp b/Engine/Plugins/Editor/WorldPartitionHLODUtilities/Source/Private/HLODBuilderMeshMerge.cpp
--- a/Engine/Plugins/Editor/WorldPartitionHLODUtilities/Source/Private/HLODBuilderMeshMerge.cpp
+++ b/Engine/Plugins/Editor/WorldPartitionHLODUtilities/Source/Private/HLODBuilderMeshMerge.cpp
@@ -83,7 +83,7 @@ TArray<UActorComponent*> UHLODBuilderMeshMerge::Build(const FHLODBuildContext& I
 	FMeshMergingSettings UseSettings = MeshMergeSettings->MeshMergeSettings; // Make a copy as we may tweak some values
 	UMaterialInterface* HLODMaterial = MeshMergeSettings->HLODMaterial.LoadSynchronous();
 
-	// When using automatic textuse sizing based on draw distance, use the MinVisibleDistance for this HLOD.
+	// When using automatic texture sizing based on draw distance, use the MinVisibleDistance for this HLOD.
 	if (UseSettings.MaterialSettings.TextureSizingType == TextureSizingType_AutomaticFromMeshDrawDistance)
 	{
 		UseSettings.MaterialSettings.MeshMinDrawDistance = InHLODBuildContext.MinVisibleDistance;
diff --git a/Engine/Plugins/Editor/WorldPartitionHLODUtilities/Source/Private/HLODBuilderMeshSimplify.cpp b/Engine/Plugins/Editor/WorldPartitionHLODUtilities/Source/Private/HLODBuilderMeshSimplify.cpp
--- a/Engine/Plugins/Editor/WorldPartitionHLODUtilities/Source/Private/HLODBuilderMeshSimplify.cpp
+++ b/Engine/Plugins/Editor/WorldPartitionHLODUtilities/Source/Private/HLODBuilderMeshSimplify.cpp
@@ -77,7 +77,7 @@ TArray<UActorComponent*> UHLODBuilderMeshSimplify::Build(const FHLODBuildContext
 	FMeshProxySettings UseSettings = MeshSimplifySettings->MeshSimplifySettings; // Make a copy as we may tweak some values
 	UMaterialInterface* HLODMaterial = MeshSimplifySettings->HLODMaterial.LoadSynchronous();
 
-	// When using automatic textuse sizing based on draw distance, use the MinVisibleDistance for this HLOD.
+	// When using automatic texture sizing based on draw distance, use the MinVisibleDistance for this HLOD.
 	if (UseSettings.MaterialSettings.TextureSizingType == TextureSizingType_AutomaticFromMeshDrawDistance)
 	{
 		UseSettings.MaterialSettings.MeshMinDrawDistance = InHLODBuildContext.MinVisibleDistance;
