WorldPartitionHLODUtilities module truncation fixes

diff --git a/Engine/Plugins/Editor/WorldPartitionHLODUtilities/Source/Private/HLODBuilderMeshSimplify.cpp b/Engine/Plugins/Editor/WorldPartitionHLODUtilities/Source/Private/HLODBuilderMeshSimplify.cpp
--- a/Engine/Plugins/Editor/WorldPartitionHLODUtilities/Source/Private/HLODBuilderMeshSimplify.cpp
+++ b/Engine/Plugins/Editor/WorldPartitionHLODUtilities/Source/Private/HLODBuilderMeshSimplify.cpp
@@ -108,8 +108,8 @@ TArray<UActorComponent*> UHLODBuilderMeshSimplify::Build(const FHLODBuildContext
 		return Bounds;
 	};
 
-	float ScreenSizePercent = ComputeBoundsScreenSize(FVector::ZeroVector, GetComponentsBounds().SphereRadius, FVector(0.0f, 0.0f, InHLODBuildContext.MinVisibleDistance), ProjectionMatrix);
-	UseSettings.ScreenSize = ScreenSizePercent * ScreenX;
+	float ScreenSizePercent = ComputeBoundsScreenSize(FVector::ZeroVector, static_cast<float>(GetComponentsBounds().SphereRadius), FVector(0.0f, 0.0f, InHLODBuildContext.MinVisibleDistance), ProjectionMatrix);
+	UseSettings.ScreenSize = FMath::RoundToInt(ScreenSizePercent * ScreenX);
 
 	TArray<UObject*> Assets;
 	FCreateProxyDelegate ProxyDelegate;
diff --git a/Engine/Plugins/Editor/WorldPartitionHLODUtilities/Source/Private/WorldPartition/HLOD/Modifiers/HLODModifierMeshDestruction.cpp b/Engine/Plugins/Editor/WorldPartitionHLODUtilities/Source/Private/WorldPartition/HLOD/Modifiers/HLODModifierMeshDestruction.cpp
--- a/Engine/Plugins/Editor/WorldPartitionHLODUtilities/Source/Private/WorldPartition/HLOD/Modifiers/HLODModifierMeshDestruction.cpp
+++ b/Engine/Plugins/Editor/WorldPartitionHLODUtilities/Source/Private/WorldPartition/HLOD/Modifiers/HLODModifierMeshDestruction.cpp
@@ -125,7 +125,7 @@ public:
 		bool bFoundParam = Instance->GetStaticSwitchParameterValue(EnableInstanceDestroyingParamName, bEnableInstanceDestroying, ParamGUID);
 		if (bFoundParam && bEnableInstanceDestroying)
 		{
-			Instance->SetScalarParameterValueEditorOnly(FMaterialParameterInfo("NumInstances"), FMath::RoundUpToPowerOfTwo(DestructibleActorsSet.Num()));
+			Instance->SetScalarParameterValueEditorOnly(FMaterialParameterInfo("NumInstances"), static_cast<float>(FMath::RoundUpToPowerOfTwo(DestructibleActorsSet.Num())));
 			Instance->BasePropertyOverrides.TwoSided = false;
 			Instance->BasePropertyOverrides.bOverride_TwoSided = true;
 
diff --git a/Engine/Plugins/Editor/WorldPartitionHLODUtilities/Source/Private/WorldPartitionHLODUtilities.cpp b/Engine/Plugins/Editor/WorldPartitionHLODUtilities/Source/Private/WorldPartitionHLODUtilities.cpp
--- a/Engine/Plugins/Editor/WorldPartitionHLODUtilities/Source/Private/WorldPartitionHLODUtilities.cpp
+++ b/Engine/Plugins/Editor/WorldPartitionHLODUtilities/Source/Private/WorldPartitionHLODUtilities.cpp
@@ -191,7 +191,7 @@ TArray<AWorldPartitionHLOD*> FWorldPartitionHLODUtilities::CreateHLODActors(FHLO
 
 		auto ComputeCellHash = [](const FString& HLODLayerName, const FGuid CellGuid)
 		{
-			uint64 HLODLayerNameHash = FCrc::StrCrc32(*HLODLayerName);
+			uint32 HLODLayerNameHash = FCrc::StrCrc32(*HLODLayerName);
 			uint32 CellGuidHash = GetTypeHash(CellGuid);
 			return HashCombine(HLODLayerNameHash, CellGuidHash);
 		};
@@ -542,7 +542,7 @@ void GatherOutputStats(AWorldPartitionHLOD* InHLODActor)
 					if (Texture && Texture->GetPackage() == HLODPackage)
 					{
 						TexturesResourceSize += Texture->GetResourceSizeBytes(EResourceSizeMode::Exclusive);
-						return Texture->GetSurfaceWidth();
+						return FMath::RoundToInt64(Texture->GetSurfaceWidth());
 					}
 				}
 
@@ -600,7 +600,7 @@ uint32 FWorldPartitionHLODUtilities::BuildHLOD(AWorldPartitionHLOD* InHLODActor)
 	{
 		FAutoScopedDurationTimer LoadTimeScope;
 		LevelStreaming = CreateLevelStreamingFromHLODActor(InHLODActor, bIsDirty);
-		LoadTimeMS = LoadTimeScope.GetTime() * 1000;
+		LoadTimeMS = FMath::RoundToInt(LoadTimeScope.GetTime() * 1000);
 	}
 
 	ON_SCOPE_EXIT
@@ -679,7 +679,7 @@ uint32 FWorldPartitionHLODUtilities::BuildHLOD(AWorldPartitionHLOD* InHLODActor)
 			{
 				FAutoScopedDurationTimer BuildTimeScope;
 				HLODComponents = HLODBuilder->Build(HLODBuildContext);
-				BuildTimeMS = BuildTimeScope.GetTime() * 1000;
+				BuildTimeMS = FMath::RoundToInt(BuildTimeScope.GetTime() * 1000);
 			}
 
 			if (HLODModifier)
@@ -754,7 +754,7 @@ uint32 FWorldPartitionHLODUtilities::BuildHLOD(AWorldPartitionHLOD* InHLODActor)
 	// Gather stats pertaining to the assets generated during this build
 	GatherOutputStats(InHLODActor);
 
-	TotalTimeMS = TotalTimeScope.GetTime() * 1000;
+	TotalTimeMS = FMath::RoundToInt(TotalTimeScope.GetTime() * 1000);
 
 	// Build timings stats
 	InHLODActor->SetStat(FWorldPartitionHLODStats::BuildTimeLoadMilliseconds, LoadTimeMS);
diff --git a/Engine/Source/Runtime/Engine/Classes/Engine/MaterialMerging.h b/Engine/Source/Runtime/Engine/Classes/Engine/MaterialMerging.h
--- a/Engine/Source/Runtime/Engine/Classes/Engine/MaterialMerging.h
+++ b/Engine/Source/Runtime/Engine/Classes/Engine/MaterialMerging.h
@@ -52,7 +52,7 @@ struct ENGINE_API FMaterialProxySettings
 
 	// Expected minimum distance at which the mesh will be rendered
 	UPROPERTY(Category = Material, BlueprintReadWrite, EditAnywhere, meta = (ClampMin = "0", EditConditionHides, EditCondition = "TextureSizingType == ETextureSizingType::TextureSizingType_AutomaticFromMeshDrawDistance"))
-	float MeshMinDrawDistance;
+	double MeshMinDrawDistance;
 	
 	// Gutter space to take into account 
 	UPROPERTY(Category = Material, BlueprintReadWrite, AdvancedDisplay, EditAnywhere)
