Pipe content bundle information from GridCellDataChunk to HlodActors creation
CreateHLODActors uses the content bundle guid to save the Hlods actors at the right places.
Currently ContentBundleGuid should always be invalid.

diff --git a/Engine/Plugins/Editor/WorldPartitionHLODUtilities/Source/Private/WorldPartitionHLODUtilities.cpp b/Engine/Plugins/Editor/WorldPartitionHLODUtilities/Source/Private/WorldPartitionHLODUtilities.cpp
--- a/Engine/Plugins/Editor/WorldPartitionHLODUtilities/Source/Private/WorldPartitionHLODUtilities.cpp
+++ b/Engine/Plugins/Editor/WorldPartitionHLODUtilities/Source/Private/WorldPartitionHLODUtilities.cpp
@@ -19,6 +19,7 @@
 #include "WorldPartition/HLOD/HLODModifier.h"
 #include "WorldPartition/HLOD/HLODStats.h"
 #include "WorldPartition/WorldPartitionLevelStreamingDynamic.h"
+#include "WorldPartition/ContentBundle/ContentBundleActivationScope.h"
 
 #include "HLODBuilderInstancing.h"
 #include "HLODBuilderMeshMerge.h"
@@ -49,8 +50,7 @@ static UWorldPartitionLevelStreamingDynamic* CreateLevelStreamingFromHLODActor(A
 	const FName LevelStreamingName = FName(*FString::Printf(TEXT("HLODLevelStreaming_%s"), *InHLODActor->GetName()));
 	TArray<FWorldPartitionRuntimeCellObjectMapping> Mappings;
 	Mappings.Reserve(InHLODActor->GetSubActors().Num());
-	// @todo_ow: investigate if we need to pass in a content bundle id here
-	Algo::Transform(InHLODActor->GetSubActors(), Mappings, [World](const FHLODSubActor& SubActor) { return FWorldPartitionRuntimeCellObjectMapping(SubActor.ActorPackage, SubActor.ActorPath, SubActor.ContainerID, SubActor.ContainerTransform, SubActor.ContainerPackage, World->GetPackage()->GetFName(), FGuid()); });
+	Algo::Transform(InHLODActor->GetSubActors(), Mappings, [World, InHLODActor](const FHLODSubActor& SubActor) { return FWorldPartitionRuntimeCellObjectMapping(SubActor.ActorPackage, SubActor.ActorPath, SubActor.ContainerID, SubActor.ContainerTransform, SubActor.ContainerPackage, World->GetPackage()->GetFName(), InHLODActor->GetContentBundleGuid()); });
 
 	UWorldPartitionLevelStreamingDynamic* LevelStreaming = UWorldPartitionLevelStreamingDynamic::LoadInEditor(World, LevelStreamingName, Mappings);
 	check(LevelStreaming);
@@ -157,7 +157,7 @@ PRAGMA_ENABLE_DEPRECATION_WARNINGS
 	}
 }
 
-TArray<AWorldPartitionHLOD*> FWorldPartitionHLODUtilities::CreateHLODActors(FHLODCreationContext& InCreationContext, const FHLODCreationParams& InCreationParams, const TArray<IStreamingGenerationContext::FActorInstance>& InActors, const TArray<const UDataLayerInstance*>& InDataLayersInstances)
+TArray<AWorldPartitionHLOD*> FWorldPartitionHLODUtilities::CreateHLODActors(FHLODCreationContext& InCreationContext, const FHLODCreationParams& InCreationParams, const TArray<IStreamingGenerationContext::FActorInstance>& InActors)
 {
 	TMap<UHLODLayer*, TSet<FHLODSubActor>> SubActorsPerHLODLayer;
 
@@ -210,16 +210,20 @@ TArray<AWorldPartitionHLOD*> FWorldPartitionHLODUtilities::CreateHLODActors(FHLO
 		bool bNewActor = HLODActor == nullptr;
 		if (bNewActor)
 		{
+			FContentBundleActivationScope Activationscope(InCreationParams.ContentBundleGuid);
+
 			FActorSpawnParameters SpawnParams;
 			SpawnParams.Name = HLODActorName;
 			SpawnParams.NameMode = FActorSpawnParameters::ESpawnActorNameMode::Required_Fatal;
 			HLODActor = InCreationParams.WorldPartition->GetWorld()->SpawnActor<AWorldPartitionHLOD>(SpawnParams);
 
+			check(HLODActor->GetContentBundleGuid() == InCreationParams.ContentBundleGuid);
+
 			HLODActor->SetSourceCellGuid(InCreationParams.CellGuid);
 			HLODActor->SetSubActorsHLODLayer(HLODLayer);
 
 			// Make sure the generated HLOD actor has the same data layers as the source actors
-			for (const UDataLayerInstance* DataLayerInstance : InDataLayersInstances)
+			for (const UDataLayerInstance* DataLayerInstance : InCreationParams.DataLayerInstances)
 			{
 				HLODActor->AddDataLayer(DataLayerInstance);
 			}
diff --git a/Engine/Plugins/Editor/WorldPartitionHLODUtilities/Source/Public/WorldPartitionHLODUtilities.h b/Engine/Plugins/Editor/WorldPartitionHLODUtilities/Source/Public/WorldPartitionHLODUtilities.h
--- a/Engine/Plugins/Editor/WorldPartitionHLODUtilities/Source/Public/WorldPartitionHLODUtilities.h
+++ b/Engine/Plugins/Editor/WorldPartitionHLODUtilities/Source/Public/WorldPartitionHLODUtilities.h
@@ -10,7 +10,7 @@
 class FWorldPartitionHLODUtilities : public IWorldPartitionHLODUtilities
 {
 public:
-	virtual TArray<AWorldPartitionHLOD*> CreateHLODActors(FHLODCreationContext& InCreationContext, const FHLODCreationParams& InCreationParams, const TArray<IStreamingGenerationContext::FActorInstance>& InActors, const TArray<const UDataLayerInstance*>& InDataLayerInstances) override;
+	virtual TArray<AWorldPartitionHLOD*> CreateHLODActors(FHLODCreationContext& InCreationContext, const FHLODCreationParams& InCreationParams, const TArray<IStreamingGenerationContext::FActorInstance>& InActors) override;
 	virtual uint32 BuildHLOD(AWorldPartitionHLOD* InHLODActor) override;
 	virtual TSubclassOf<UHLODBuilder> GetHLODBuilderClass(const UHLODLayer* InHLODLayer) override;
 	virtual UHLODBuilderSettings* CreateHLODBuilderSettings(UHLODLayer* InHLODLayer) override;
diff --git a/Engine/Source/Runtime/Engine/Private/WorldPartition/RuntimeSpatialHash/RuntimeSpatialHashHLOD.cpp b/Engine/Source/Runtime/Engine/Private/WorldPartition/RuntimeSpatialHash/RuntimeSpatialHashHLOD.cpp
--- a/Engine/Source/Runtime/Engine/Private/WorldPartition/RuntimeSpatialHash/RuntimeSpatialHashHLOD.cpp
+++ b/Engine/Source/Runtime/Engine/Private/WorldPartition/RuntimeSpatialHash/RuntimeSpatialHashHLOD.cpp
@@ -174,10 +174,12 @@ static TArray<FGuid> GenerateHLODsForGrid(UWorldPartition* WorldPartition, const
 			if (bShouldGenerateHLODs)
 			{
 				SlowTask.EnterProgressFrame(1);
-				// todo_ow : Generating hlods is not yet supported for content bundles. Pass an invalid UID to specified the cell is not coming from ContentBundles.
-				FGuid ContentBundleInvalidUID;
-				const FGuid CellGuid = UWorldPartitionRuntimeSpatialHash::GetCellGuid(RuntimeGrid.GridName, CellGlobalCoord, GridCellDataChunk.GetDataLayersID(), ContentBundleInvalidUID);
-				const FString CellName = UWorldPartitionRuntimeSpatialHash::GetCellNameString(OuterWorld->GetPackage()->GetFName(), RuntimeGrid.GridName, CellGlobalCoord, GridCellDataChunk.GetDataLayersID(), ContentBundleInvalidUID);
+				
+				// todo_ow: Content Bundle Hlods generation is not enable currently. We should never get a valid ContentBundleGuid.
+				check(!GridCellDataChunk.GetContentBundleID().IsValid());
+
+				const FGuid CellGuid = UWorldPartitionRuntimeSpatialHash::GetCellGuid(RuntimeGrid.GridName, CellGlobalCoord, GridCellDataChunk.GetDataLayersID(), GridCellDataChunk.GetContentBundleID());
+				const FString CellName = UWorldPartitionRuntimeSpatialHash::GetCellNameString(OuterWorld->GetPackage()->GetFName(), RuntimeGrid.GridName, CellGlobalCoord, GridCellDataChunk.GetDataLayersID(), GridCellDataChunk.GetContentBundleID());
 
 				UE_LOG(LogWorldPartitionRuntimeSpatialHashHLOD, Verbose, TEXT("Creating HLOD for cell %s at %s..."), *CellName, *CellCoord.ToString());
 
@@ -213,9 +215,11 @@ static TArray<FGuid> GenerateHLODsForGrid(UWorldPartition* WorldPartition, const
 				CreationParams.CellBounds = CellBounds;
 				CreationParams.HLODLevel = HLODLevel;
 				CreationParams.MinVisibleDistance = RuntimeGrid.LoadingRange;
+				CreationParams.ContentBundleGuid = GridCellDataChunk.GetContentBundleID();
+				CreationParams.DataLayerInstances = GridCellDataChunk.GetDataLayers();
 
 				IWorldPartitionHLODUtilities* WPHLODUtilities = FModuleManager::Get().LoadModuleChecked<IWorldPartitionHLODUtilitiesModule>("WorldPartitionHLODUtilities").GetUtilities();
-				TArray<AWorldPartitionHLOD*> CellHLODActors = WPHLODUtilities->CreateHLODActors(Context, CreationParams, ActorInstances, GridCellDataChunk.GetDataLayers());
+				TArray<AWorldPartitionHLOD*> CellHLODActors = WPHLODUtilities->CreateHLODActors(Context, CreationParams, ActorInstances);
 				if (!CellHLODActors.IsEmpty())
 				{
 					TArray<AWorldPartitionHLOD*> NewCellHLODActors;
diff --git a/Engine/Source/Runtime/Engine/Public/WorldPartition/HLOD/IWorldPartitionHLODUtilities.h b/Engine/Source/Runtime/Engine/Public/WorldPartition/HLOD/IWorldPartitionHLODUtilities.h
--- a/Engine/Source/Runtime/Engine/Public/WorldPartition/HLOD/IWorldPartitionHLODUtilities.h
+++ b/Engine/Source/Runtime/Engine/Public/WorldPartition/HLOD/IWorldPartitionHLODUtilities.h
@@ -30,6 +30,8 @@ struct ENGINE_API FHLODCreationParams
 	FGuid CellGuid;
 	FBox  CellBounds;
 	uint32 HLODLevel;
+	FGuid ContentBundleGuid;
+	TArray<const UDataLayerInstance*> DataLayerInstances;
 
 	double MinVisibleDistance;
 };
@@ -50,7 +52,7 @@ public:
 	 * @param	InActors			The actors for which we'll build an HLOD representation
 	 * @param	InDataLayers		The data layers to assign to the newly created HLOD actors
 	 */
-	virtual TArray<AWorldPartitionHLOD*> CreateHLODActors(FHLODCreationContext& InCreationContext, const FHLODCreationParams& InCreationParams, const TArray<IStreamingGenerationContext::FActorInstance>& InActors, const TArray<const UDataLayerInstance*>& InDataLayerInstances) = 0;
+	virtual TArray<AWorldPartitionHLOD*> CreateHLODActors(FHLODCreationContext& InCreationContext, const FHLODCreationParams& InCreationParams, const TArray<IStreamingGenerationContext::FActorInstance>& InActors) = 0;
 
 	/**
 	 * Build HLOD for the specified AWorldPartitionHLOD actor.
@@ -75,4 +77,12 @@ public:
 	 * @return A newly created UHLODBuilderSettings object, outered to the provided HLOD layer.
 	 */
 	virtual UHLODBuilderSettings* CreateHLODBuilderSettings(UHLODLayer* InHLODLayer) = 0;
+
+
+	UE_DEPRECATED(5.2, "Use the overload that passes the DataLayersInstances via InCreationParams")
+	virtual TArray<AWorldPartitionHLOD*> CreateHLODActors(FHLODCreationContext& InCreationContext, const FHLODCreationParams& InCreationParams, const TArray<IStreamingGenerationContext::FActorInstance>& InActors, const TArray<const UDataLayerInstance*>& InDataLayerInstances)
+	{
+		const_cast<FHLODCreationParams&>(InCreationParams).DataLayerInstances = InDataLayerInstances;
+		return CreateHLODActors(InCreationContext, InCreationParams, InActors);
+	}
 };
