World Partition - HLOD: Revert to using the cell name rather than it's GUID for the HLOD actor's label

diff --git a/Engine/Plugins/Editor/WorldPartitionHLODUtilities/Source/Private/WorldPartitionHLODUtilities.cpp b/Engine/Plugins/Editor/WorldPartitionHLODUtilities/Source/Private/WorldPartitionHLODUtilities.cpp
--- a/Engine/Plugins/Editor/WorldPartitionHLODUtilities/Source/Private/WorldPartitionHLODUtilities.cpp
+++ b/Engine/Plugins/Editor/WorldPartitionHLODUtilities/Source/Private/WorldPartitionHLODUtilities.cpp
@@ -295,7 +295,7 @@ TArray<AWorldPartitionHLOD*> FWorldPartitionHLODUtilities::CreateHLODActors(FHLO
 		}
 
 		// Actor label
-		const FString ActorLabel = FString::Printf(TEXT("%s/%s"), *HLODLayer->GetName(), *InCreationParams.CellGuid.ToString());
+		const FString ActorLabel = FString::Printf(TEXT("%s/%s"), *HLODLayer->GetName(), *InCreationParams.CellName);
 		if (HLODActor->GetActorLabel() != ActorLabel)
 		{
 			HLODActor->SetActorLabel(ActorLabel);
diff --git a/Engine/Source/Runtime/Engine/Private/WorldPartition/RuntimeSpatialHash/RuntimeSpatialHashHLOD.cpp b/Engine/Source/Runtime/Engine/Private/WorldPartition/RuntimeSpatialHash/RuntimeSpatialHashHLOD.cpp
--- a/Engine/Source/Runtime/Engine/Private/WorldPartition/RuntimeSpatialHash/RuntimeSpatialHashHLOD.cpp
+++ b/Engine/Source/Runtime/Engine/Private/WorldPartition/RuntimeSpatialHash/RuntimeSpatialHashHLOD.cpp
@@ -211,6 +211,7 @@ static TArray<FGuid> GenerateHLODsForGrid(UWorldPartition* WorldPartition, const
 				FHLODCreationParams CreationParams;
 				CreationParams.WorldPartition = WorldPartition;
 				CreationParams.CellGuid = CellGuid;
+				CreationParams.CellName = CellName;
 				CreationParams.CellBounds = CellBounds;
 				CreationParams.HLODLevel = HLODLevel;
 				CreationParams.MinVisibleDistance = RuntimeGrid.LoadingRange;
diff --git a/Engine/Source/Runtime/Engine/Public/WorldPartition/HLOD/IWorldPartitionHLODUtilities.h b/Engine/Source/Runtime/Engine/Public/WorldPartition/HLOD/IWorldPartitionHLODUtilities.h
--- a/Engine/Source/Runtime/Engine/Public/WorldPartition/HLOD/IWorldPartitionHLODUtilities.h
+++ b/Engine/Source/Runtime/Engine/Public/WorldPartition/HLOD/IWorldPartitionHLODUtilities.h
@@ -28,6 +28,7 @@ struct ENGINE_API FHLODCreationParams
 	FName HLODLayerName;
 
 	FGuid CellGuid;
+	FString CellName;
 	FBox  CellBounds;
 	uint32 HLODLevel;
 	FGuid ContentBundleGuid;
