WorldPartition - HLOD: Added a base key for instanced HLOD, forces a rebuild of all instanced HLODs

diff --git a/Engine/Plugins/Editor/WorldPartitionHLODUtilities/Source/Private/HLODBuilderInstancing.cpp b/Engine/Plugins/Editor/WorldPartitionHLODUtilities/Source/Private/HLODBuilderInstancing.cpp
--- a/Engine/Plugins/Editor/WorldPartitionHLODUtilities/Source/Private/HLODBuilderInstancing.cpp
+++ b/Engine/Plugins/Editor/WorldPartitionHLODUtilities/Source/Private/HLODBuilderInstancing.cpp
@@ -20,6 +20,10 @@ uint32 UHLODBuilderInstancingSettings::GetCRC() const
 
 	FArchiveCrc32 Ar;
 
+	// Base key, changing this will force a rebuild of all HLODs from this builder
+	FString HLODBaseKey = "53809597CD9C4FB7AC75827A628513D6";
+	Ar << HLODBaseKey;
+
 	Ar << This.bDisallowNanite;
 	UE_LOG(LogHLODBuilder, VeryVerbose, TEXT(" - bDisallowNanite = %d"), Ar.GetCrc());
 
