Reorganized source files under WorldPartitionHLODUtilities

diff --git a/Engine/Plugins/Editor/WorldPartitionHLODUtilities/Source/Private/HLODBuilderInstancing.cpp b/Engine/Plugins/Editor/WorldPartitionHLODUtilities/Source/Private/HLODBuilderInstancing.cpp
deleted file mode 100644
diff --git a/Engine/Plugins/Editor/WorldPartitionHLODUtilities/Source/Private/HLODBuilderMeshApproximate.cpp b/Engine/Plugins/Editor/WorldPartitionHLODUtilities/Source/Private/HLODBuilderMeshApproximate.cpp
deleted file mode 100644
diff --git a/Engine/Plugins/Editor/WorldPartitionHLODUtilities/Source/Private/HLODBuilderMeshMerge.cpp b/Engine/Plugins/Editor/WorldPartitionHLODUtilities/Source/Private/HLODBuilderMeshMerge.cpp
deleted file mode 100644
diff --git a/Engine/Plugins/Editor/WorldPartitionHLODUtilities/Source/Private/HLODBuilderMeshSimplify.cpp b/Engine/Plugins/Editor/WorldPartitionHLODUtilities/Source/Private/HLODBuilderMeshSimplify.cpp
deleted file mode 100644
diff --git a/Engine/Plugins/Editor/WorldPartitionHLODUtilities/Source/Private/WorldPartition/HLOD/Builders/HLODBuilderInstancing.cpp b/Engine/Plugins/Editor/WorldPartitionHLODUtilities/Source/Private/WorldPartition/HLOD/Builders/HLODBuilderInstancing.cpp
new file mode 100644
--- /dev/null
+++ b/Engine/Plugins/Editor/WorldPartitionHLODUtilities/Source/Private/WorldPartition/HLOD/Builders/HLODBuilderInstancing.cpp
@@ -0,0 +1,71 @@
+// Copyright Epic Games, Inc. All Rights Reserved.
+
+#include "WorldPartition/HLOD/Builders/HLODBuilderInstancing.h"
+#include "Serialization/ArchiveCrc32.h"
+#include "Engine/StaticMesh.h"
+#include "Components/StaticMeshComponent.h"
+
+#include UE_INLINE_GENERATED_CPP_BY_NAME(HLODBuilderInstancing)
+
+
+UHLODBuilderInstancingSettings::UHLODBuilderInstancingSettings(const FObjectInitializer& ObjectInitializer)
+	: Super(ObjectInitializer)
+	, bDisallowNanite(false)
+{
+}
+
+uint32 UHLODBuilderInstancingSettings::GetCRC() const
+{
+	UHLODBuilderInstancingSettings& This = *const_cast<UHLODBuilderInstancingSettings*>(this);
+
+	FArchiveCrc32 Ar;
+
+	// Base key, changing this will force a rebuild of all HLODs from this builder
+	FString HLODBaseKey = "53809597CD9C4FB7AC75827A628513D6";
+	Ar << HLODBaseKey;
+
+	Ar << This.bDisallowNanite;
+	UE_LOG(LogHLODBuilder, VeryVerbose, TEXT(" - bDisallowNanite = %d"), Ar.GetCrc());
+
+	return Ar.GetCrc();
+}
+
+
+UHLODBuilderInstancing::UHLODBuilderInstancing(const FObjectInitializer& ObjectInitializer)
+	: Super(ObjectInitializer)
+{
+}
+
+TSubclassOf<UHLODBuilderSettings> UHLODBuilderInstancing::GetSettingsClass() const
+{
+	return UHLODBuilderInstancingSettings::StaticClass();
+}
+
+TArray<UActorComponent*> UHLODBuilderInstancing::Build(const FHLODBuildContext& InHLODBuildContext, const TArray<UActorComponent*>& InSourceComponents) const
+{
+	const UHLODBuilderInstancingSettings& InstancingSettings = *CastChecked<UHLODBuilderInstancingSettings>(HLODBuilderSettings);
+
+	TArray<UActorComponent*> HLODComponents = UHLODBuilder::BatchInstances(InSourceComponents);
+
+	// If requested, disallow Nanite on components whose mesh is Nanite enabled and has multiple LODs
+	if (InstancingSettings.bDisallowNanite)
+	{
+		for (UActorComponent* HLODComponent : HLODComponents)
+		{
+			if (UStaticMeshComponent* SMComponent = Cast<UStaticMeshComponent>(HLODComponent))
+			{
+				if (UStaticMesh* StaticMesh = SMComponent->GetStaticMesh())
+				{
+					if (!SMComponent->bDisallowNanite && StaticMesh->HasValidNaniteData() && StaticMesh->GetNumLODs() > 1)
+					{
+						SMComponent->bDisallowNanite = true;
+						SMComponent->MarkRenderStateDirty();
+					}
+				}
+			}
+		}
+	}
+
+	return HLODComponents;
+}
+
diff --git a/Engine/Plugins/Editor/WorldPartitionHLODUtilities/Source/Private/WorldPartition/HLOD/Builders/HLODBuilderMeshApproximate.cpp b/Engine/Plugins/Editor/WorldPartitionHLODUtilities/Source/Private/WorldPartition/HLOD/Builders/HLODBuilderMeshApproximate.cpp
new file mode 100644
--- /dev/null
+++ b/Engine/Plugins/Editor/WorldPartitionHLODUtilities/Source/Private/WorldPartition/HLOD/Builders/HLODBuilderMeshApproximate.cpp
@@ -0,0 +1,183 @@
+// Copyright Epic Games, Inc. All Rights Reserved.
+
+#include "WorldPartition/HLOD/Builders/HLODBuilderMeshApproximate.h"
+
+#include "Algo/ForEach.h"
+#include "Components/StaticMeshComponent.h"
+#include "UObject/Package.h"
+
+#include "Engine/Engine.h"
+#include "Engine/StaticMesh.h"
+#include "Engine/Texture2D.h"
+#include "Materials/Material.h"
+#include "Materials/MaterialInstance.h"
+
+#include "Modules/ModuleManager.h"
+#include "IGeometryProcessingInterfacesModule.h"
+#include "GeometryProcessingInterfaces/ApproximateActors.h"
+
+#include "Materials/Material.h"
+#include "Engine/HLODProxy.h"
+#include "Serialization/ArchiveCrc32.h"
+
+#include UE_INLINE_GENERATED_CPP_BY_NAME(HLODBuilderMeshApproximate)
+
+
+UHLODBuilderMeshApproximate::UHLODBuilderMeshApproximate(const FObjectInitializer& ObjectInitializer)
+	: Super(ObjectInitializer)
+{
+}
+
+UHLODBuilderMeshApproximateSettings::UHLODBuilderMeshApproximateSettings(const FObjectInitializer& ObjectInitializer)
+	: Super(ObjectInitializer)
+{
+#if WITH_EDITORONLY_DATA
+	if (!IsTemplate())
+	{
+		HLODMaterial = GEngine->DefaultHLODFlattenMaterial;
+	}
+#endif
+}
+
+uint32 UHLODBuilderMeshApproximateSettings::GetCRC() const
+{
+	UHLODBuilderMeshApproximateSettings& This = *const_cast<UHLODBuilderMeshApproximateSettings*>(this);
+
+	FArchiveCrc32 Ar;
+
+	Ar << This.MeshApproximationSettings;
+	UE_LOG(LogHLODBuilder, VeryVerbose, TEXT(" - MeshApproximationSettings = %d"), Ar.GetCrc());
+
+	uint32 Hash = Ar.GetCrc();
+
+	if (!HLODMaterial.IsNull())
+	{
+		UMaterialInterface* Material = HLODMaterial.LoadSynchronous();
+		if (Material)
+		{
+			uint32 MaterialCRC = UHLODProxy::GetCRC(Material);
+			UE_LOG(LogHLODBuilder, VeryVerbose, TEXT(" - Material = %d"), MaterialCRC);
+			Hash = HashCombine(Hash, MaterialCRC);
+		}
+	}
+
+	return Hash;
+}
+
+TSubclassOf<UHLODBuilderSettings> UHLODBuilderMeshApproximate::GetSettingsClass() const
+{
+	return UHLODBuilderMeshApproximateSettings::StaticClass();
+}
+
+TArray<UActorComponent*> UHLODBuilderMeshApproximate::Build(const FHLODBuildContext& InHLODBuildContext, const TArray<UActorComponent*>& InSourceComponents) const
+{
+	TRACE_CPUPROFILER_EVENT_SCOPE(UHLODBuilderMeshApproximate::Build);
+
+	TArray<UPrimitiveComponent*> PrimitiveComponents = FilterComponents<UPrimitiveComponent>(InSourceComponents);
+
+	TSet<AActor*> Actors;
+	Algo::Transform(PrimitiveComponents, Actors, [](UPrimitiveComponent* PrimitiveComponent) { return PrimitiveComponent->GetOwner(); });
+	
+	IGeometryProcessingInterfacesModule& GeomProcInterfaces = FModuleManager::Get().LoadModuleChecked<IGeometryProcessingInterfacesModule>("GeometryProcessingInterfaces");
+	IGeometryProcessing_ApproximateActors* ApproxActorsAPI = GeomProcInterfaces.GetApproximateActorsImplementation();
+
+	//
+	// Construct options for ApproximateActors operation
+	//
+
+	const UHLODBuilderMeshApproximateSettings* MeshApproximateSettings = CastChecked<UHLODBuilderMeshApproximateSettings>(HLODBuilderSettings);
+	FMeshApproximationSettings UseSettings = MeshApproximateSettings->MeshApproximationSettings; // Make a copy as we may tweak some values
+	UMaterialInterface* HLODMaterial = MeshApproximateSettings->HLODMaterial.LoadSynchronous();
+
+	// When using automatic texture sizing based on draw distance, use the MinVisibleDistance for this HLOD.
+	if (UseSettings.MaterialSettings.TextureSizingType == TextureSizingType_AutomaticFromMeshDrawDistance)
+	{
+		UseSettings.MaterialSettings.MeshMinDrawDistance = InHLODBuildContext.MinVisibleDistance;
+	}
+
+	IGeometryProcessing_ApproximateActors::FOptions Options = ApproxActorsAPI->ConstructOptions(UseSettings);
+	Options.BasePackagePath = InHLODBuildContext.AssetsOuter->GetPackage()->GetName();
+	Options.bGenerateLightmapUVs = false;
+	Options.bCreatePhysicsBody = false;
+
+	// Material baking settings
+	Options.BakeMaterial = HLODMaterial;
+	Options.BaseColorTexParamName = FName("BaseColorTexture");
+	Options.NormalTexParamName = FName("NormalTexture");
+	Options.MetallicTexParamName = FName("MetallicTexture");
+	Options.RoughnessTexParamName = FName("RoughnessTexture");
+	Options.SpecularTexParamName = FName("SpecularTexture");
+	Options.EmissiveTexParamName = FName("EmissiveHDRTexture");
+	Options.bUsePackedMRS = true;
+	Options.PackedMRSTexParamName = FName("PackedTexture");
+
+	// Compute texel density if needed, depending on the TextureSizingType setting
+	if (UseSettings.MaterialSettings.ResolveTexelDensity(PrimitiveComponents, Options.MeshTexelDensity))
+	{ 
+		Options.TextureSizePolicy = IGeometryProcessing_ApproximateActors::ETextureSizePolicy::TexelDensity;
+	}
+
+	// run actor approximation computation
+	IGeometryProcessing_ApproximateActors::FResults Results;
+	ApproxActorsAPI->ApproximateActors(Actors.Array(), Options, Results);
+
+	TArray<UActorComponent*> Components;
+	if (Results.ResultCode == IGeometryProcessing_ApproximateActors::EResultCode::Success)
+	{
+		auto FixupAsset = [InHLODBuildContext](UObject* Asset)
+		{
+			Asset->ClearFlags(RF_Public | RF_Standalone);
+			Asset->Rename(nullptr, InHLODBuildContext.AssetsOuter, REN_NonTransactional | REN_DontCreateRedirectors | REN_ForceNoResetLoaders);
+		};
+	
+		Algo::ForEach(Results.NewMeshAssets, FixupAsset);
+		Algo::ForEach(Results.NewMaterials, FixupAsset);
+		Algo::ForEach(Results.NewTextures, FixupAsset);
+
+		for (UStaticMesh* StaticMesh : Results.NewMeshAssets)
+		{
+			UStaticMeshComponent* Component = NewObject<UStaticMeshComponent>();
+			Component->SetStaticMesh(StaticMesh);
+
+			Components.Add(Component);
+		}
+
+		for (UMaterialInterface* Material : Results.NewMaterials)
+		{
+			UMaterialInstance* MaterialInst = CastChecked<UMaterialInstance>(Material);
+
+			FStaticParameterSet StaticParameterSet;
+			
+			auto SetStaticSwitch = [&StaticParameterSet](FName ParamName, bool bSet)
+			{
+				if (bSet)
+				{
+					FStaticSwitchParameter SwitchParameter;
+					SwitchParameter.ParameterInfo.Name = ParamName;
+					SwitchParameter.Value = true;
+					SwitchParameter.bOverride = true;
+					StaticParameterSet.StaticSwitchParameters.Add(SwitchParameter);
+				}
+			};
+
+			// Set proper switches needed by our base flatten material
+			SetStaticSwitch("UseBaseColor", Options.bBakeBaseColor);
+			SetStaticSwitch("UseRoughness", Options.bBakeRoughness);
+			SetStaticSwitch("UseMetallic", Options.bBakeMetallic);
+			SetStaticSwitch("UseSpecular", Options.bBakeSpecular);
+			SetStaticSwitch("UseEmissiveHDR", Options.bBakeEmissive);
+			SetStaticSwitch("UseNormal", Options.bBakeNormalMap);
+			SetStaticSwitch("PackMetallic", Options.bUsePackedMRS);
+			SetStaticSwitch("PackSpecular", Options.bUsePackedMRS);
+			SetStaticSwitch("PackRoughness", Options.bUsePackedMRS);
+
+			// Force initializing the static permutations according to the switches we have set
+			MaterialInst->UpdateStaticPermutation(StaticParameterSet);
+			MaterialInst->InitStaticPermutation();
+			MaterialInst->PostEditChange();
+		}
+	}
+
+	return Components;
+}
+
diff --git a/Engine/Plugins/Editor/WorldPartitionHLODUtilities/Source/Private/WorldPartition/HLOD/Builders/HLODBuilderMeshMerge.cpp b/Engine/Plugins/Editor/WorldPartitionHLODUtilities/Source/Private/WorldPartition/HLOD/Builders/HLODBuilderMeshMerge.cpp
new file mode 100644
--- /dev/null
+++ b/Engine/Plugins/Editor/WorldPartitionHLODUtilities/Source/Private/WorldPartition/HLOD/Builders/HLODBuilderMeshMerge.cpp
@@ -0,0 +1,117 @@
+// Copyright Epic Games, Inc. All Rights Reserved.
+
+#include "WorldPartition/HLOD/Builders/HLODBuilderMeshMerge.h"
+
+#include "Algo/ForEach.h"
+#include "Engine/Engine.h"
+#include "Engine/StaticMesh.h"
+#include "Materials/Material.h"
+#include "Components/StaticMeshComponent.h"
+
+#include "IMaterialBakingModule.h"
+#include "IMeshMergeUtilities.h"
+#include "MeshMergeModule.h"
+#include "Modules/ModuleManager.h"
+
+#include "Materials/Material.h"
+#include "Engine/HLODProxy.h"
+#include "Serialization/ArchiveCrc32.h"
+
+#include UE_INLINE_GENERATED_CPP_BY_NAME(HLODBuilderMeshMerge)
+
+
+UHLODBuilderMeshMerge::UHLODBuilderMeshMerge(const FObjectInitializer& ObjectInitializer)
+	: Super(ObjectInitializer)
+{
+}
+
+UHLODBuilderMeshMergeSettings::UHLODBuilderMeshMergeSettings(const FObjectInitializer& ObjectInitializer)
+	: Super(ObjectInitializer)
+{
+#if WITH_EDITORONLY_DATA
+	if (!IsTemplate())
+	{
+		HLODMaterial = GEngine->DefaultHLODFlattenMaterial;
+	}
+#endif
+}
+
+uint32 UHLODBuilderMeshMergeSettings::GetCRC() const
+{
+	UHLODBuilderMeshMergeSettings& This = *const_cast<UHLODBuilderMeshMergeSettings*>(this);
+
+	FArchiveCrc32 Ar;
+
+	// Base mesh merge key, changing this will force a rebuild of all HLODs from this builder
+	FString HLODBaseKey = "89D89284DD3847FA90C5998E06DD8FEC";
+	Ar << HLODBaseKey;
+
+	Ar << This.MeshMergeSettings;
+	UE_LOG(LogHLODBuilder, VeryVerbose, TEXT(" - MeshMergeSettings = %d"), Ar.GetCrc());
+
+	if (MeshMergeSettings.bMergeMaterials)
+	{
+		IMaterialBakingModule& Module = FModuleManager::Get().LoadModuleChecked<IMaterialBakingModule>("MaterialBaking");
+		uint32 MaterialBakingModuleCRC = Module.GetCRC();
+		Ar << MaterialBakingModuleCRC;
+	}
+
+	uint32 Hash = Ar.GetCrc();
+
+	if (!HLODMaterial.IsNull())
+	{
+		UMaterialInterface* Material = HLODMaterial.LoadSynchronous();
+		if (Material)
+		{
+			uint32 MaterialCRC = UHLODProxy::GetCRC(Material);
+			UE_LOG(LogHLODBuilder, VeryVerbose, TEXT(" - Material = %d"), MaterialCRC);
+			Hash = HashCombine(Hash, MaterialCRC);
+		}
+	}
+
+	return Hash;
+}
+
+TSubclassOf<UHLODBuilderSettings> UHLODBuilderMeshMerge::GetSettingsClass() const
+{
+	return UHLODBuilderMeshMergeSettings::StaticClass();
+}
+
+TArray<UActorComponent*> UHLODBuilderMeshMerge::Build(const FHLODBuildContext& InHLODBuildContext, const TArray<UActorComponent*>& InSourceComponents) const
+{
+	TRACE_CPUPROFILER_EVENT_SCOPE(UHLODBuilderMeshMerge::CreateComponents);
+
+	TArray<UPrimitiveComponent*> SourcePrimitiveComponents = FilterComponents<UPrimitiveComponent>(InSourceComponents);
+
+	TArray<UObject*> Assets;
+	FVector MergedActorLocation;
+
+	const UHLODBuilderMeshMergeSettings* MeshMergeSettings = CastChecked<UHLODBuilderMeshMergeSettings>(HLODBuilderSettings);
+	FMeshMergingSettings UseSettings = MeshMergeSettings->MeshMergeSettings; // Make a copy as we may tweak some values
+	UMaterialInterface* HLODMaterial = MeshMergeSettings->HLODMaterial.LoadSynchronous();
+
+	// When using automatic texture sizing based on draw distance, use the MinVisibleDistance for this HLOD.
+	if (UseSettings.MaterialSettings.TextureSizingType == TextureSizingType_AutomaticFromMeshDrawDistance)
+	{
+		UseSettings.MaterialSettings.MeshMinDrawDistance = InHLODBuildContext.MinVisibleDistance;
+	}
+
+	const IMeshMergeUtilities& MeshMergeUtilities = FModuleManager::Get().LoadModuleChecked<IMeshMergeModule>("MeshMergeUtilities").GetUtilities();
+	MeshMergeUtilities.MergeComponentsToStaticMesh(SourcePrimitiveComponents, InHLODBuildContext.World, UseSettings, HLODMaterial, InHLODBuildContext.AssetsOuter->GetPackage(), InHLODBuildContext.AssetsBaseName, Assets, MergedActorLocation, 0.25f, false);
+
+	UStaticMeshComponent* Component = nullptr;
+	Algo::ForEach(Assets, [this, &Component, &MergedActorLocation](UObject* Asset)
+	{
+		Asset->ClearFlags(RF_Public | RF_Standalone);
+
+		if (Cast<UStaticMesh>(Asset))
+		{
+			Component = NewObject<UStaticMeshComponent>();
+			Component->SetStaticMesh(static_cast<UStaticMesh*>(Asset));
+			Component->SetWorldLocation(MergedActorLocation);
+		}
+	});
+
+	return { Component };
+}
+
diff --git a/Engine/Plugins/Editor/WorldPartitionHLODUtilities/Source/Private/WorldPartition/HLOD/Builders/HLODBuilderMeshSimplify.cpp b/Engine/Plugins/Editor/WorldPartitionHLODUtilities/Source/Private/WorldPartition/HLOD/Builders/HLODBuilderMeshSimplify.cpp
new file mode 100644
--- /dev/null
+++ b/Engine/Plugins/Editor/WorldPartitionHLODUtilities/Source/Private/WorldPartition/HLOD/Builders/HLODBuilderMeshSimplify.cpp
@@ -0,0 +1,147 @@
+// Copyright Epic Games, Inc. All Rights Reserved.
+
+#include "WorldPartition/HLOD/Builders/HLODBuilderMeshSimplify.h"
+
+#include "Algo/ForEach.h"
+#include "Engine/Engine.h"
+#include "Engine/StaticMesh.h"
+#include "Materials/Material.h"
+#include "Components/StaticMeshComponent.h"
+
+#include "IMaterialBakingModule.h"
+#include "IMeshMergeUtilities.h"
+#include "MeshMergeModule.h"
+#include "Modules/ModuleManager.h"
+
+#include "Materials/Material.h"
+#include "Engine/HLODProxy.h"
+#include "SceneManagement.h"
+#include "Serialization/ArchiveCrc32.h"
+
+#include UE_INLINE_GENERATED_CPP_BY_NAME(HLODBuilderMeshSimplify)
+
+
+UHLODBuilderMeshSimplify::UHLODBuilderMeshSimplify(const FObjectInitializer& ObjectInitializer)
+	: Super(ObjectInitializer)
+{
+}
+
+UHLODBuilderMeshSimplifySettings::UHLODBuilderMeshSimplifySettings(const FObjectInitializer& ObjectInitializer)
+	: Super(ObjectInitializer)
+{
+#if WITH_EDITORONLY_DATA
+	if (!IsTemplate())
+	{
+		HLODMaterial = GEngine->DefaultHLODFlattenMaterial;
+	}
+#endif
+}
+
+uint32 UHLODBuilderMeshSimplifySettings::GetCRC() const
+{
+	UHLODBuilderMeshSimplifySettings& This = *const_cast<UHLODBuilderMeshSimplifySettings*>(this);
+
+	FArchiveCrc32 Ar;
+
+	// Base mesh simplify key, changing this will force a rebuild of all HLODs from this builder
+	FString HLODBaseKey = "9015D801CD0E420B9D2B0CB997A97813";
+	Ar << HLODBaseKey;
+
+	Ar << This.MeshSimplifySettings;
+	UE_LOG(LogHLODBuilder, VeryVerbose, TEXT(" - MeshSimplifySettings = %d"), Ar.GetCrc());
+
+	IMaterialBakingModule& Module = FModuleManager::Get().LoadModuleChecked<IMaterialBakingModule>("MaterialBaking");
+	uint32 MaterialBakingModuleCRC = Module.GetCRC();
+	Ar << MaterialBakingModuleCRC;
+
+	uint32 Hash = Ar.GetCrc();
+
+	if (!HLODMaterial.IsNull())
+	{
+		UMaterialInterface* Material = HLODMaterial.LoadSynchronous();
+		if (Material)
+		{
+			uint32 MaterialCRC = UHLODProxy::GetCRC(Material);
+			UE_LOG(LogHLODBuilder, VeryVerbose, TEXT(" - Material = %d"), MaterialCRC);
+			Hash = HashCombine(Hash, MaterialCRC);
+		}
+	}
+
+	return Hash;
+}
+
+TSubclassOf<UHLODBuilderSettings> UHLODBuilderMeshSimplify::GetSettingsClass() const
+{
+	return UHLODBuilderMeshSimplifySettings::StaticClass();
+}
+
+TArray<UActorComponent*> UHLODBuilderMeshSimplify::Build(const FHLODBuildContext& InHLODBuildContext, const TArray<UActorComponent*>& InSourceComponents) const
+{
+	TRACE_CPUPROFILER_EVENT_SCOPE(UHLODBuilderMeshSimplifySettings::Build);
+
+	TArray<UStaticMeshComponent*> StaticMeshComponents = FilterComponents<UStaticMeshComponent>(InSourceComponents);
+
+	const UHLODBuilderMeshSimplifySettings* MeshSimplifySettings = CastChecked<UHLODBuilderMeshSimplifySettings>(HLODBuilderSettings);
+	FMeshProxySettings UseSettings = MeshSimplifySettings->MeshSimplifySettings; // Make a copy as we may tweak some values
+	UMaterialInterface* HLODMaterial = MeshSimplifySettings->HLODMaterial.LoadSynchronous();
+
+	// When using automatic texture sizing based on draw distance, use the MinVisibleDistance for this HLOD.
+	if (UseSettings.MaterialSettings.TextureSizingType == TextureSizingType_AutomaticFromMeshDrawDistance)
+	{
+		UseSettings.MaterialSettings.MeshMinDrawDistance = InHLODBuildContext.MinVisibleDistance;
+	}
+
+	// Generate a projection matrix.
+	static const float ScreenX = 1920;
+	static const float ScreenY = 1080;
+	static const float HalfFOVRad = FMath::DegreesToRadians(45.0f);
+	static const FMatrix ProjectionMatrix = FPerspectiveMatrix(HalfFOVRad, ScreenX, ScreenY, 0.01f);
+
+	// Gather bounds of the input components
+	auto GetComponentsBounds = [&]() -> FBoxSphereBounds
+	{
+		FBoxSphereBounds Bounds;
+		bool bFirst = true;
+
+		for (UActorComponent* Component : InSourceComponents)
+		{
+			if (USceneComponent* SceneComponent = Cast<USceneComponent>(Component))
+			{
+				FBoxSphereBounds ComponentBounds = SceneComponent->Bounds;
+				Bounds = bFirst ? ComponentBounds : Bounds + ComponentBounds;
+				bFirst = false;
+			}
+		}
+
+		return Bounds;
+	};
+
+	float ScreenSizePercent = ComputeBoundsScreenSize(FVector::ZeroVector, static_cast<float>(GetComponentsBounds().SphereRadius), FVector(0.0f, 0.0f, InHLODBuildContext.MinVisibleDistance), ProjectionMatrix);
+	UseSettings.ScreenSize = FMath::RoundToInt(ScreenSizePercent * ScreenX);
+
+	TArray<UObject*> Assets;
+	FCreateProxyDelegate ProxyDelegate;
+	ProxyDelegate.BindLambda([&Assets](const FGuid Guid, TArray<UObject*>& InAssetsCreated) { Assets = InAssetsCreated; });
+
+	const IMeshMergeUtilities& MeshMergeUtilities = FModuleManager::Get().LoadModuleChecked<IMeshMergeModule>("MeshMergeUtilities").GetUtilities();
+	MeshMergeUtilities.CreateProxyMesh(StaticMeshComponents, UseSettings, HLODMaterial, InHLODBuildContext.AssetsOuter->GetPackage(), InHLODBuildContext.AssetsBaseName, FGuid::NewGuid(), ProxyDelegate, true);
+
+
+	TArray<UActorComponent*> Components;
+
+	Algo::ForEach(Assets, [this, &Components](UObject* Asset)
+	{
+		Asset->ClearFlags(RF_Public | RF_Standalone);
+
+		if (Cast<UStaticMesh>(Asset))
+		{
+			UStaticMeshComponent* SMComponent = NewObject<UStaticMeshComponent>();
+			SMComponent->SetStaticMesh(static_cast<UStaticMesh*>(Asset));
+
+			Components.Add(SMComponent);
+		}
+	});
+	
+	return Components;
+}
+
diff --git a/Engine/Plugins/Editor/WorldPartitionHLODUtilities/Source/Private/WorldPartition/HLOD/Utilities/WorldPartitionHLODUtilities.cpp b/Engine/Plugins/Editor/WorldPartitionHLODUtilities/Source/Private/WorldPartition/HLOD/Utilities/WorldPartitionHLODUtilities.cpp
new file mode 100644
--- /dev/null
+++ b/Engine/Plugins/Editor/WorldPartitionHLODUtilities/Source/Private/WorldPartition/HLOD/Utilities/WorldPartitionHLODUtilities.cpp
@@ -0,0 +1,771 @@
+// Copyright Epic Games, Inc. All Rights Reserved.
+
+#include "WorldPartition/HLOD/Utilities/WorldPartitionHLODUtilities.h"
+#include "WorldPartition/WorldPartitionActorDescView.h"
+#include "WorldPartition/WorldPartitionStreamingGeneration.h"
+
+#if WITH_EDITOR
+
+#include "BodySetupEnums.h"
+#include "WorldPartition/WorldPartition.h"
+#include "Components/InstancedStaticMeshComponent.h"
+#include "WorldPartition/DataLayer/DataLayerInstance.h"
+#include "Engine/CollisionProfile.h"
+#include "WorldPartition/HLOD/HLODActor.h"
+#include "Engine/Texture.h"
+#include "WorldPartition/HLOD/HLODSubActor.h"
+#include "StaticMeshResources.h"
+#include "WorldPartition/HLOD/HLODLayer.h"
+#include "WorldPartition/HLOD/HLODModifier.h"
+#include "WorldPartition/HLOD/HLODStats.h"
+#include "WorldPartition/WorldPartitionLevelStreamingDynamic.h"
+#include "WorldPartition/ContentBundle/ContentBundleActivationScope.h"
+
+#include "WorldPartition/HLOD/Builders/HLODBuilderInstancing.h"
+#include "WorldPartition/HLOD/Builders/HLODBuilderMeshMerge.h"
+#include "WorldPartition/HLOD/Builders/HLODBuilderMeshSimplify.h"
+#include "WorldPartition/HLOD/Builders/HLODBuilderMeshApproximate.h"
+
+#include "Engine/Level.h"
+#include "Engine/StaticMesh.h"
+#include "Materials/MaterialInstance.h"
+#include "PhysicsEngine/BodySetup.h"
+#include "ProfilingDebugging/ScopedTimers.h"
+#include "Serialization/ArchiveCrc32.h"
+#include "StaticMeshCompiler.h"
+#include "TextureCompiler.h"
+#include "UObject/MetaData.h"
+#include "UObject/GCObjectScopeGuard.h"
+
+
+static UWorldPartitionLevelStreamingDynamic* CreateLevelStreamingFromHLODActor(AWorldPartitionHLOD* InHLODActor, bool& bOutDirty)
+{
+	UPackage::WaitForAsyncFileWrites();
+
+	bOutDirty = false;
+	UWorld* World = InHLODActor->GetWorld();
+	UWorldPartition* WorldPartition = World->GetWorldPartition();
+	check(WorldPartition);
+	
+	const FName LevelStreamingName = FName(*FString::Printf(TEXT("HLODLevelStreaming_%s"), *InHLODActor->GetName()));
+	TArray<FWorldPartitionRuntimeCellObjectMapping> Mappings;
+	Mappings.Reserve(InHLODActor->GetSubActors().Num());
+	Algo::Transform(InHLODActor->GetSubActors(), Mappings, [World, InHLODActor](const FHLODSubActor& SubActor) { return FWorldPartitionRuntimeCellObjectMapping(SubActor.ActorPackage, SubActor.ActorPath, SubActor.ContainerID, SubActor.ContainerTransform, SubActor.ContainerPackage, World->GetPackage()->GetFName(), InHLODActor->GetContentBundleGuid()); });
+
+	UWorldPartitionLevelStreamingDynamic* LevelStreaming = UWorldPartitionLevelStreamingDynamic::LoadInEditor(World, LevelStreamingName, Mappings);
+	check(LevelStreaming);
+
+	if (LevelStreaming->GetLoadSucceeded())
+	{
+		// Finish assets compilation
+		FAssetCompilingManager::Get().FinishAllCompilation();
+
+		// Ensure all deferred construction scripts are executed
+		FAssetCompilingManager::Get().ProcessAsyncTasks();
+	}
+	else
+	{
+		bOutDirty = true;
+		UE_LOG(LogHLODBuilder, Warning, TEXT("HLOD actor \"%s\" needs to be rebuilt as it didn't succeed in loading all actors."), *InHLODActor->GetActorLabel());
+	}
+
+	return LevelStreaming;
+}
+
+static uint32 GetCRC(const UHLODLayer* InHLODLayer)
+{
+	UHLODLayer& HLODLayer= *const_cast<UHLODLayer*>(InHLODLayer);
+
+	uint32 CRC;
+
+	CRC = GetTypeHash(HLODLayer.GetLayerType());
+	UE_LOG(LogHLODBuilder, VeryVerbose, TEXT(" - LayerType = %d"), CRC);
+
+	CRC = HashCombine(HLODLayer.GetHLODBuilderSettings()->GetCRC(), CRC);
+	UE_LOG(LogHLODBuilder, VeryVerbose, TEXT(" - HLODBuilderSettings = %d"), CRC);
+
+	CRC = HashCombine(HLODLayer.GetCellSize(), CRC);
+	UE_LOG(LogHLODBuilder, VeryVerbose, TEXT(" - CellSize = %d"), CRC);
+
+	return CRC;
+}
+
+static uint32 GetCRC(const TArray<FHLODSubActor>& InSubActors)
+{
+	uint32 CRC = 0;
+
+	for (const FHLODSubActor& HLODSubActor : InSubActors)
+	{
+		CRC = HashCombine(GetTypeHash(HLODSubActor), CRC);
+	}
+
+	return CRC;
+}
+
+static uint32 ComputeHLODHash(AWorldPartitionHLOD* InHLODActor, const TArray<AActor*>& InActors)
+{
+	FArchiveCrc32 Ar;
+
+	// Base key, changing this will force a rebuild of all HLODs
+	FString HLODBaseKey = "A6C1517FC49B478E833A540E2C06289F";
+	Ar << HLODBaseKey;
+
+	// Sub Actors
+	uint32 SubActorsHash = GetCRC(InHLODActor->GetSubActors());
+	UE_LOG(LogHLODBuilder, VeryVerbose, TEXT(" - Sub Actors (%d actors) = %x"), InHLODActor->GetSubActors().Num(), SubActorsHash);
+	Ar << SubActorsHash;
+
+	// HLOD Layer
+	uint32 HLODLayerHash = GetCRC(InHLODActor->GetSubActorsHLODLayer());
+	UE_LOG(LogHLODBuilder, VeryVerbose, TEXT(" - HLOD Layer (%s) = %x"), *InHLODActor->GetSubActorsHLODLayer()->GetName(), HLODLayerHash);
+	Ar << HLODLayerHash;
+
+	// Min Visible Distance
+	uint32 HLODMinVisibleDistanceHash = GetTypeHash(InHLODActor->GetMinVisibleDistance());
+	UE_LOG(LogHLODBuilder, VeryVerbose, TEXT(" - HLOD Min Visible Distance (%.02f) = %x"), InHLODActor->GetMinVisibleDistance(), HLODMinVisibleDistanceHash);
+	Ar << HLODMinVisibleDistanceHash;
+
+	// Append all components CRCs
+	uint32 HLODComponentsHash = UHLODBuilder::ComputeHLODHash(InActors);
+	UE_LOG(LogHLODBuilder, VeryVerbose, TEXT(" - HLOD Source Components = %x"), HLODComponentsHash);
+	Ar << HLODComponentsHash;
+
+	return Ar.GetCrc();
+}
+
+void AddSubActor(const FWorldPartitionActorDescView& ActorDescView, const IStreamingGenerationContext::FActorInstance& ActorInstance, TSet<FHLODSubActor>& SubActors)
+{
+PRAGMA_DISABLE_DEPRECATION_WARNINGS
+	// Leaving this as deprecated for now until we fix up the serialization format for SubActorsInfo and do proper upgrade/deprecation 
+	FName ActorPath = ActorDescView.GetActorPath();
+PRAGMA_ENABLE_DEPRECATION_WARNINGS
+
+	// Add the actor
+	bool bIsAlreadyInSet = false;
+	FHLODSubActor SubActor(ActorDescView.GetGuid(), ActorDescView.GetActorPackage(), ActorPath, ActorInstance.GetContainerID(), ActorInstance.GetActorDescContainer()->GetContainerPackage(), ActorInstance.GetTransform());
+	SubActors.Add(SubActor, &bIsAlreadyInSet);
+
+	if (!bIsAlreadyInSet)
+	{
+		// Add its references
+		const FActorDescViewMap* ActorDescViewMap = ActorInstance.ActorSetInstance->ContainerInstance->ActorDescViewMap;
+		for (const FGuid& ReferenceGuid : ActorDescView.GetReferences())
+		{
+			const FWorldPartitionActorDescView& RefActorDescView = ActorDescViewMap->FindByGuidChecked(ReferenceGuid);
+			AddSubActor(RefActorDescView, ActorInstance, SubActors);
+		}
+	}
+}
+
+TArray<AWorldPartitionHLOD*> FWorldPartitionHLODUtilities::CreateHLODActors(FHLODCreationContext& InCreationContext, const FHLODCreationParams& InCreationParams, const TArray<IStreamingGenerationContext::FActorInstance>& InActors)
+{
+	TMap<UHLODLayer*, TSet<FHLODSubActor>> SubActorsPerHLODLayer;
+
+	for (const IStreamingGenerationContext::FActorInstance& ActorInstance : InActors)
+	{
+		const FWorldPartitionActorDescView& ActorDescView = ActorInstance.GetActorDescView();
+
+		if (ActorDescView.GetActorIsHLODRelevant())
+		{
+			if (!ActorInstance.ActorSetInstance->bIsSpatiallyLoaded)
+			{
+				UE_LOG(LogHLODBuilder, Warning, TEXT("Tried to included non-spatially loaded actor %s into HLOD"), *ActorDescView.GetActorName().ToString());
+				continue;
+			}
+
+			UHLODLayer* HLODLayer = UHLODLayer::GetHLODLayer(ActorDescView, InCreationParams.WorldPartition);
+			if (HLODLayer)
+			{
+				TSet<FHLODSubActor>& SubActors = SubActorsPerHLODLayer.FindOrAdd(HLODLayer);
+				AddSubActor(ActorDescView, ActorInstance, SubActors);
+			}
+		}
+	}
+
+	TArray<AWorldPartitionHLOD*> HLODActors;
+	for (const auto& Pair : SubActorsPerHLODLayer)
+	{
+		const UHLODLayer* HLODLayer = Pair.Key;
+		const TSet<FHLODSubActor>& SubActors = Pair.Value;
+		check(!SubActors.IsEmpty());
+
+		auto ComputeCellHash = [](const FString& HLODLayerName, const FGuid CellGuid)
+		{
+			uint32 HLODLayerNameHash = FCrc::StrCrc32(*HLODLayerName);
+			uint32 CellGuidHash = GetTypeHash(CellGuid);
+			return HashCombine(HLODLayerNameHash, CellGuidHash);
+		};
+
+		uint64 CellHash = ComputeCellHash(HLODLayer->GetName(), InCreationParams.CellGuid);
+		FName HLODActorName = *FString::Printf(TEXT("%s_%016llx"), *HLODLayer->GetName(), CellHash);		
+
+		AWorldPartitionHLOD* HLODActor = nullptr;
+		FWorldPartitionHandle HLODActorHandle;
+		if (InCreationContext.HLODActorDescs.RemoveAndCopyValue(HLODActorName, HLODActorHandle))
+		{
+			InCreationContext.ActorReferences.Add(HLODActorHandle.ToReference());
+			HLODActor = CastChecked<AWorldPartitionHLOD>(HLODActorHandle->GetActor());
+		}
+
+		bool bNewActor = HLODActor == nullptr;
+		if (bNewActor)
+		{
+			FContentBundleActivationScope Activationscope(InCreationParams.ContentBundleGuid);
+
+			FActorSpawnParameters SpawnParams;
+			SpawnParams.Name = HLODActorName;
+			SpawnParams.NameMode = FActorSpawnParameters::ESpawnActorNameMode::Required_Fatal;
+			HLODActor = InCreationParams.WorldPartition->GetWorld()->SpawnActor<AWorldPartitionHLOD>(SpawnParams);
+
+			check(HLODActor->GetContentBundleGuid() == InCreationParams.ContentBundleGuid);
+
+			HLODActor->SetSourceCellGuid(InCreationParams.CellGuid);
+			HLODActor->SetSubActorsHLODLayer(HLODLayer);
+
+			// Make sure the generated HLOD actor has the same data layers as the source actors
+			for (const UDataLayerInstance* DataLayerInstance : InCreationParams.DataLayerInstances)
+			{
+				HLODActor->AddDataLayer(DataLayerInstance);
+			}
+		}
+		else
+		{
+#if DO_CHECK
+			check(HLODActor->GetSourceCellGuid() == InCreationParams.CellGuid);
+			check(HLODActor->GetSubActorsHLODLayer() == HLODLayer);
+#endif
+		}
+
+		bool bIsDirty = false;
+
+		// Sub actors
+		{
+			bool bSubActorsChanged = HLODActor->GetSubActors().Num() != SubActors.Num();
+			if (!bSubActorsChanged)
+			{
+				TArray<FHLODSubActor> A = HLODActor->GetSubActors();
+				TArray<FHLODSubActor> B = SubActors.Array();
+				A.Sort();
+				B.Sort();
+				bSubActorsChanged = A != B;
+			}
+
+			if (bSubActorsChanged)
+			{
+				HLODActor->SetSubActors(SubActors.Array());
+				bIsDirty = true;
+			}
+		}
+
+		// Runtime grid
+		FName RuntimeGrid = HLODLayer->GetRuntimeGrid(InCreationParams.HLODLevel);
+		if (HLODActor->GetRuntimeGrid() != RuntimeGrid)
+		{
+			HLODActor->SetRuntimeGrid(RuntimeGrid);
+			bIsDirty = true;
+		}
+
+		// Spatially loaded
+		if (HLODActor->GetIsSpatiallyLoaded() != HLODLayer->IsSpatiallyLoaded())
+		{
+			HLODActor->SetIsSpatiallyLoaded(HLODLayer->IsSpatiallyLoaded());
+			bIsDirty = true;
+		}
+
+		// HLOD level
+		if (HLODActor->GetLODLevel() != InCreationParams.HLODLevel)
+		{
+			HLODActor->SetLODLevel(InCreationParams.HLODLevel);
+			bIsDirty = true;
+		}
+
+		// Require warmup
+		if (HLODActor->DoesRequireWarmup() != HLODLayer->DoesRequireWarmup())
+		{
+			HLODActor->SetRequireWarmup(HLODLayer->DoesRequireWarmup());
+			bIsDirty = true;
+		}
+
+		// Parent HLOD layer
+		UHLODLayer* ParentHLODLayer = HLODLayer->GetParentLayer().LoadSynchronous();
+		if (HLODActor->GetHLODLayer() != ParentHLODLayer)
+		{
+			HLODActor->SetHLODLayer(ParentHLODLayer);
+			bIsDirty = true;
+		}
+
+		// Actor label
+		const FString ActorLabel = FString::Printf(TEXT("%s/%s"), *HLODLayer->GetName(), *InCreationParams.CellName);
+		if (HLODActor->GetActorLabel() != ActorLabel)
+		{
+			HLODActor->SetActorLabel(ActorLabel);
+			bIsDirty = true;
+		}
+
+		// Folder name
+		const FName FolderPath(FString::Printf(TEXT("HLOD/%s"), *HLODLayer->GetName()));
+		if (HLODActor->GetFolderPath() != FolderPath)
+		{
+			HLODActor->SetFolderPath(FolderPath);
+			bIsDirty = true;
+		}
+
+		// Cell bounds
+		if (!HLODActor->GetHLODBounds().Equals(InCreationParams.CellBounds))
+		{
+			HLODActor->SetHLODBounds(InCreationParams.CellBounds);
+			bIsDirty = true;
+		}
+
+		// Minimum visible distance
+		if (!FMath::IsNearlyEqual(HLODActor->GetMinVisibleDistance(), InCreationParams.MinVisibleDistance))
+		{
+			HLODActor->SetMinVisibleDistance(InCreationParams.MinVisibleDistance);
+			bIsDirty = true;
+		}
+
+		// If any change was performed, mark HLOD package as dirty
+		if (bIsDirty)
+		{
+			HLODActor->MarkPackageDirty();
+		}
+
+		HLODActors.Add(HLODActor);
+	}
+
+	return HLODActors;
+}
+
+TSubclassOf<UHLODBuilder> FWorldPartitionHLODUtilities::GetHLODBuilderClass(const UHLODLayer* InHLODLayer)
+{
+	EHLODLayerType HLODLayerType = InHLODLayer->GetLayerType();
+	switch (HLODLayerType)
+	{
+	case EHLODLayerType::Instancing:
+		return UHLODBuilderInstancing::StaticClass();
+		break;
+
+	case EHLODLayerType::MeshMerge:
+		return UHLODBuilderMeshMerge::StaticClass();
+		break;
+
+	case EHLODLayerType::MeshSimplify:
+		return UHLODBuilderMeshSimplify::StaticClass();
+		break;
+
+	case EHLODLayerType::MeshApproximate:
+		return UHLODBuilderMeshApproximate::StaticClass();
+		break;
+
+	case EHLODLayerType::Custom:
+		return InHLODLayer->GetHLODBuilderClass();
+		break;
+
+	default:
+		checkf(false, TEXT("Unsupported type"));
+		return nullptr;
+	}
+}
+
+UHLODBuilderSettings* FWorldPartitionHLODUtilities::CreateHLODBuilderSettings(UHLODLayer* InHLODLayer)
+{
+	// Retrieve the HLOD builder class
+	TSubclassOf<UHLODBuilder> HLODBuilderClass = GetHLODBuilderClass(InHLODLayer);
+	if (!HLODBuilderClass)
+	{
+		return NewObject<UHLODBuilderSettings>(InHLODLayer, UHLODBuilderSettings::StaticClass());
+	}
+
+	// Retrieve the HLOD builder settings class
+	TSubclassOf<UHLODBuilderSettings> HLODBuilderSettingsClass = HLODBuilderClass->GetDefaultObject<UHLODBuilder>()->GetSettingsClass();
+	if (!ensure(HLODBuilderSettingsClass))
+	{
+		return NewObject<UHLODBuilderSettings>(InHLODLayer, UHLODBuilderSettings::StaticClass());
+	}
+
+	UHLODBuilderSettings* HLODBuilderSettings = NewObject<UHLODBuilderSettings>(InHLODLayer, HLODBuilderSettingsClass);
+
+	// Deprecated properties handling
+	if (InHLODLayer->GetHLODBuilderSettings() == nullptr)
+	{
+		EHLODLayerType HLODLayerType = InHLODLayer->GetLayerType();
+		switch (HLODLayerType)
+		{
+		case EHLODLayerType::MeshMerge:
+			CastChecked<UHLODBuilderMeshMergeSettings>(HLODBuilderSettings)->MeshMergeSettings = InHLODLayer->MeshMergeSettings_DEPRECATED;
+			CastChecked<UHLODBuilderMeshMergeSettings>(HLODBuilderSettings)->HLODMaterial = InHLODLayer->HLODMaterial_DEPRECATED;
+			break;
+
+		case EHLODLayerType::MeshSimplify:
+			CastChecked<UHLODBuilderMeshSimplifySettings>(HLODBuilderSettings)->MeshSimplifySettings = InHLODLayer->MeshSimplifySettings_DEPRECATED;
+			CastChecked<UHLODBuilderMeshSimplifySettings>(HLODBuilderSettings)->HLODMaterial = InHLODLayer->HLODMaterial_DEPRECATED;
+			break;
+
+		case EHLODLayerType::MeshApproximate:
+			CastChecked<UHLODBuilderMeshApproximateSettings>(HLODBuilderSettings)->MeshApproximationSettings = InHLODLayer->MeshApproximationSettings_DEPRECATED;
+			CastChecked<UHLODBuilderMeshApproximateSettings>(HLODBuilderSettings)->HLODMaterial = InHLODLayer->HLODMaterial_DEPRECATED;
+			break;
+		};
+	}
+
+	return HLODBuilderSettings;
+}
+
+void GatherInputStats(AWorldPartitionHLOD* InHLODActor, const TArray<AActor*> SubActors)
+{
+	int64 NumActors = 0;
+	int64 NumTriangles = 0;
+	int64 NumVertices = 0;
+
+	for (AActor* SubActor : SubActors)
+	{
+		if (SubActor && SubActor->IsHLODRelevant())
+		{
+			if (AWorldPartitionHLOD* SubHLODActor = Cast<AWorldPartitionHLOD>(SubActor))
+			{
+				NumActors += SubHLODActor->GetStat(FWorldPartitionHLODStats::InputActorCount);
+				NumTriangles += SubHLODActor->GetStat(FWorldPartitionHLODStats::InputTriangleCount);
+				NumVertices += SubHLODActor->GetStat(FWorldPartitionHLODStats::InputVertexCount);				
+			}
+			else
+			{
+				NumActors++;
+
+				SubActor->ForEachComponent<UStaticMeshComponent>(false, [&] (UStaticMeshComponent* SMComponent)
+				{
+					const UStaticMesh* StaticMesh = SMComponent->GetStaticMesh();
+					const FStaticMeshRenderData* RenderData = StaticMesh ? StaticMesh->GetRenderData() : nullptr;
+					const bool bHasRenderData = RenderData && !RenderData->LODResources.IsEmpty();
+
+					if (bHasRenderData)
+					{
+						const UInstancedStaticMeshComponent* ISMComponent = Cast<UInstancedStaticMeshComponent>(SMComponent);
+						const int64 LOD0TriCount = RenderData->LODResources[0].GetNumTriangles();
+						const int64 LOD0VtxCount = RenderData->LODResources[0].GetNumVertices();
+						const int64 NumInstances = ISMComponent ? ISMComponent->GetInstanceCount() : 1;
+
+						NumTriangles += LOD0TriCount * NumInstances;
+						NumVertices += LOD0VtxCount * NumInstances;
+					}
+				});
+			}
+		}
+	}
+
+	InHLODActor->SetStat(FWorldPartitionHLODStats::InputActorCount, NumActors);
+	InHLODActor->SetStat(FWorldPartitionHLODStats::InputTriangleCount, NumTriangles);
+	InHLODActor->SetStat(FWorldPartitionHLODStats::InputVertexCount, NumVertices);	
+}
+
+void GatherOutputStats(AWorldPartitionHLOD* InHLODActor)
+{
+	const UPackage* HLODPackage = InHLODActor->GetPackage();
+
+	// Gather relevant assets and process them outside of this ForEach as it's possible that async compilation completion
+	// triggers insertion of new objects, which would break the iteration over UObjects
+	TArray<UStaticMesh*> StaticMeshes;
+	TArray<UTexture*> Textures;
+	TArray<UMaterialInstance*> MaterialInstances;
+	ForEachObjectWithPackage(HLODPackage, [&](UObject* Object)
+	{
+		if (UStaticMesh* StaticMesh = Cast<UStaticMesh>(Object))
+		{
+			StaticMeshes.Add(StaticMesh);
+		}
+		else if (UMaterialInstance* MaterialInstance = Cast<UMaterialInstance>(Object))
+		{
+			MaterialInstances.Add(MaterialInstance);
+		}
+		else if (UTexture* Texture = Cast<UTexture>(Object))
+		{
+			Textures.Add(Texture);
+		}
+		return true;
+	}, false);
+
+	// Process static meshes
+	int64 MeshResourceSize = 0;
+	{
+		int64 InstanceCount = 0;
+		int64 NaniteTriangleCount = 0;
+		int64 NaniteVertexCount = 0;
+		int64 TriangleCount = 0;
+		int64 VertexCount = 0;
+		int64 UVChannelCount = 0;
+
+		InHLODActor->ForEachComponent<UInstancedStaticMeshComponent>(false, [&](const UInstancedStaticMeshComponent* ISMC)
+		{
+			InstanceCount += ISMC->GetInstanceCount();
+		});
+
+		FStaticMeshCompilingManager::Get().FinishCompilation(StaticMeshes);
+
+		for (UStaticMesh* StaticMesh : StaticMeshes)
+		{
+			MeshResourceSize += StaticMesh->GetResourceSizeBytes(EResourceSizeMode::Exclusive);
+
+			TriangleCount += StaticMesh->GetNumTriangles(0);
+			VertexCount += StaticMesh->GetNumVertices(0);
+			UVChannelCount = FMath::Max(UVChannelCount, StaticMesh->GetNumTexCoords(0));
+
+			NaniteTriangleCount += StaticMesh->GetNumNaniteTriangles();
+			NaniteVertexCount += StaticMesh->GetNumNaniteVertices();
+		}
+
+		// Mesh stats
+		InHLODActor->SetStat(FWorldPartitionHLODStats::MeshInstanceCount, InstanceCount);
+		InHLODActor->SetStat(FWorldPartitionHLODStats::MeshNaniteTriangleCount, NaniteTriangleCount);
+		InHLODActor->SetStat(FWorldPartitionHLODStats::MeshNaniteVertexCount, NaniteVertexCount);
+		InHLODActor->SetStat(FWorldPartitionHLODStats::MeshTriangleCount, TriangleCount);
+		InHLODActor->SetStat(FWorldPartitionHLODStats::MeshVertexCount, VertexCount);
+		InHLODActor->SetStat(FWorldPartitionHLODStats::MeshUVChannelCount, UVChannelCount);
+	}
+
+	// Process materials
+	int64 TexturesResourceSize = 0;
+	{
+		int64 BaseColorTextureSize = 0;
+		int64 NormalTextureSize = 0;
+		int64 EmissiveTextureSize = 0;
+		int64 MetallicTextureSize = 0;
+		int64 RoughnessTextureSize = 0;
+		int64 SpecularTextureSize = 0;
+
+		FTextureCompilingManager::Get().FinishCompilation(Textures);
+
+		for (UMaterialInstance* MaterialInstance : MaterialInstances)
+		{
+			// Retrieve the texture size for a texture that can have different names
+			auto GetTextureSize = [&](const TArray<FName>& TextureParamNames) -> int64
+			{
+				for (FName TextureParamName : TextureParamNames)
+				{
+					UTexture* Texture = nullptr;
+					MaterialInstance->GetTextureParameterValue(TextureParamName, Texture, true);
+
+					if (Texture && Texture->GetPackage() == HLODPackage)
+					{
+						TexturesResourceSize += Texture->GetResourceSizeBytes(EResourceSizeMode::Exclusive);
+						return FMath::RoundToInt64(Texture->GetSurfaceWidth());
+					}
+				}
+
+				return 0;
+			};
+
+			int64 LocalBaseColorTextureSize = GetTextureSize({ "BaseColorTexture", "DiffuseTexture" });
+			int64 LocalNormalTextureSize = GetTextureSize({ "NormalTexture" });
+			int64 LocalEmissiveTextureSize = GetTextureSize({ "EmissiveTexture", "EmissiveColorTexture" });
+			int64 LocalMetallicTextureSize = GetTextureSize({ "MetallicTexture" });
+			int64 LocalRoughnessTextureSize = GetTextureSize({ "RoughnessTexture" });
+			int64 LocalSpecularTextureSize = GetTextureSize({ "SpecularTexture" });
+
+			int64 MRSTextureSize = GetTextureSize({ "PackedTexture" });
+			if (MRSTextureSize != 0)
+			{
+				LocalMetallicTextureSize = LocalRoughnessTextureSize = LocalSpecularTextureSize = MRSTextureSize;
+			}
+
+			BaseColorTextureSize = FMath::Max(BaseColorTextureSize, LocalBaseColorTextureSize);
+			NormalTextureSize = FMath::Max(NormalTextureSize, LocalNormalTextureSize);
+			EmissiveTextureSize = FMath::Max(EmissiveTextureSize, LocalEmissiveTextureSize);
+			MetallicTextureSize = FMath::Max(MetallicTextureSize, LocalMetallicTextureSize);
+			RoughnessTextureSize = FMath::Max(RoughnessTextureSize, LocalRoughnessTextureSize);
+			SpecularTextureSize = FMath::Max(SpecularTextureSize, LocalSpecularTextureSize);
+		}
+
+		// Material stats
+		InHLODActor->SetStat(FWorldPartitionHLODStats::MaterialBaseColorTextureSize, BaseColorTextureSize);
+		InHLODActor->SetStat(FWorldPartitionHLODStats::MaterialNormalTextureSize, NormalTextureSize);
+		InHLODActor->SetStat(FWorldPartitionHLODStats::MaterialEmissiveTextureSize, EmissiveTextureSize);
+		InHLODActor->SetStat(FWorldPartitionHLODStats::MaterialMetallicTextureSize, MetallicTextureSize);
+		InHLODActor->SetStat(FWorldPartitionHLODStats::MaterialRoughnessTextureSize, RoughnessTextureSize);
+		InHLODActor->SetStat(FWorldPartitionHLODStats::MaterialSpecularTextureSize, SpecularTextureSize);
+	}
+
+	// Memory stats
+	InHLODActor->SetStat(FWorldPartitionHLODStats::MemoryMeshResourceSizeBytes, MeshResourceSize);
+	InHLODActor->SetStat(FWorldPartitionHLODStats::MemoryTexturesResourceSizeBytes, TexturesResourceSize);
+	InHLODActor->SetStat(FWorldPartitionHLODStats::MemoryDiskSizeBytes, 0); // Clear disk size as it is unknown at this point. It will be assigned at package loading
+}
+
+uint32 FWorldPartitionHLODUtilities::BuildHLOD(AWorldPartitionHLOD* InHLODActor)
+{
+	FAutoScopedDurationTimer TotalTimeScope;
+
+	// Keep track of timings related to this build
+	int64 LoadTimeMS = 0;
+	int64 BuildTimeMS = 0;
+	int64 TotalTimeMS = 0;
+	
+	// Load actors relevant to HLODs
+	bool bIsDirty = false;
+	UWorldPartitionLevelStreamingDynamic* LevelStreaming = nullptr;
+	{
+		FAutoScopedDurationTimer LoadTimeScope;
+		LevelStreaming = CreateLevelStreamingFromHLODActor(InHLODActor, bIsDirty);
+		LoadTimeMS = FMath::RoundToInt(LoadTimeScope.GetTime() * 1000);
+	}
+
+	ON_SCOPE_EXIT
+	{
+		UWorldPartitionLevelStreamingDynamic::UnloadFromEditor(LevelStreaming);
+	};
+
+	uint32 OldHLODHash = bIsDirty ? 0 : InHLODActor->GetHLODHash();
+	uint32 NewHLODHash = ComputeHLODHash(InHLODActor, LevelStreaming->GetLoadedLevel()->Actors);
+
+	if (OldHLODHash == NewHLODHash)
+	{
+		UE_LOG(LogHLODBuilder, Verbose, TEXT("HLOD actor \"%s\" doesn't need to be rebuilt."), *InHLODActor->GetActorLabel());
+		return OldHLODHash;
+	}
+
+	// Clear stats as we're about to refresh them
+	InHLODActor->ResetStats();
+
+	// Rename previous assets found in the HLOD actor package.
+	// Move the previous asset(s) to the transient package, to avoid any object reuse during the build
+	{
+	    TArray<UObject*> ObjectsToRename;
+	    ForEachObjectWithOuter(InHLODActor->GetPackage(), [&ObjectsToRename](UObject* Obj)
+	    {
+		    if (!Obj->IsA<AActor>() && !Obj->IsA<UMetaData>())
+		    {
+			    ObjectsToRename.Add(Obj);
+		    }
+	    }, false);
+    
+	    
+	    for (UObject* Obj : ObjectsToRename)
+	    {
+		    // Make sure the old object is not used by anything
+		    Obj->ClearFlags(RF_Standalone | RF_Public);
+		    const FName OldRenamed = MakeUniqueObjectName(GetTransientPackage(), Obj->GetClass(), *FString::Printf(TEXT("OLD_%s"), *Obj->GetName()));
+		    Obj->Rename(*OldRenamed.ToString(), GetTransientPackage(), REN_DontCreateRedirectors | REN_NonTransactional | REN_ForceNoResetLoaders);
+	    }
+	}
+
+	// Gather stats from the input to our HLOD build
+	GatherInputStats(InHLODActor, LevelStreaming->GetLoadedLevel()->Actors);
+	
+	const UHLODLayer* HLODLayer = InHLODActor->GetSubActorsHLODLayer();
+	TSubclassOf<UHLODBuilder> HLODBuilderClass = GetHLODBuilderClass(HLODLayer);
+
+	if (HLODBuilderClass)
+	{
+		UHLODBuilder* HLODBuilder = NewObject<UHLODBuilder>(GetTransientPackage(), HLODBuilderClass);
+		if (ensure(HLODBuilder))
+		{
+			FGCObjectScopeGuard HLODBuilderGCScopeGuard(HLODBuilder);
+
+			HLODBuilder->SetHLODBuilderSettings(HLODLayer->GetHLODBuilderSettings());
+
+			FHLODBuildContext HLODBuildContext;
+			HLODBuildContext.World = InHLODActor->GetWorld();
+			HLODBuildContext.SourceActors = LevelStreaming->GetLoadedLevel()->Actors;
+			HLODBuildContext.AssetsOuter = InHLODActor->GetPackage();
+			HLODBuildContext.AssetsBaseName = InHLODActor->GetActorLabel();
+			HLODBuildContext.MinVisibleDistance = InHLODActor->GetMinVisibleDistance();
+			HLODBuildContext.WorldPosition = InHLODActor->GetHLODBounds().GetCenter();
+
+			const TSubclassOf<UWorldPartitionHLODModifier> HLODModifierClass = HLODLayer->GetHLODModifierClass();
+			UWorldPartitionHLODModifier* HLODModifier = HLODModifierClass.Get() ? NewObject<UWorldPartitionHLODModifier>(GetTransientPackage(), HLODModifierClass) : nullptr;
+			FGCObjectScopeGuard HLODModifierGCScopeGuard(HLODModifier);
+
+			if (HLODModifier)
+			{
+				HLODModifier->BeginHLODBuild(HLODBuildContext);
+			}
+
+			// Build
+			TArray<UActorComponent*> HLODComponents;
+			{
+				FAutoScopedDurationTimer BuildTimeScope;
+				HLODComponents = HLODBuilder->Build(HLODBuildContext);
+				BuildTimeMS = FMath::RoundToInt(BuildTimeScope.GetTime() * 1000);
+			}
+
+			if (HLODModifier)
+			{
+				HLODModifier->EndHLODBuild(HLODComponents);
+			}
+
+			if (HLODComponents.IsEmpty())
+			{
+				UE_LOG(LogHLODBuilder, Warning, TEXT("HLOD generation created no component for %s"), *InHLODActor->GetActorLabel());
+			}
+
+			// Ideally, this should be performed elsewhere, to allow more flexibility in the HLOD generation
+			for (UActorComponent* HLODComponent : HLODComponents)
+			{
+				HLODComponent->SetCanEverAffectNavigation(false);
+
+				if (USceneComponent* SceneComponent = Cast<USceneComponent>(HLODComponent))
+				{
+					// Change Mobility to be Static
+					SceneComponent->SetMobility(EComponentMobility::Static);
+
+					// Enable bounds optimizations
+					SceneComponent->bComputeFastLocalBounds = true;
+					SceneComponent->bComputeBoundsOnceForGame = true;
+				}
+
+				if (UPrimitiveComponent* PrimitiveComponent = Cast<UPrimitiveComponent>(HLODComponent))
+				{
+					// Disable collisions
+					PrimitiveComponent->SetCollisionProfileName(UCollisionProfile::NoCollision_ProfileName);
+					PrimitiveComponent->SetGenerateOverlapEvents(false);
+					PrimitiveComponent->CanCharacterStepUpOn = ECanBeCharacterBase::ECB_No;
+					PrimitiveComponent->SetCollisionEnabled(ECollisionEnabled::NoCollision);
+
+					// HLOD visual components aren't needed on servers
+					PrimitiveComponent->AlwaysLoadOnServer = false;
+				}
+
+				if (UStaticMeshComponent* StaticMeshComponent = Cast<UStaticMeshComponent>(HLODComponent))
+				{
+					if (UStaticMesh* StaticMesh = StaticMeshComponent->GetStaticMesh())
+					{
+						// If the HLOD process did create this static mesh
+						if (StaticMesh->GetPackage() == HLODBuildContext.AssetsOuter)
+						{
+							// Set up ray tracing far fields for always loaded HLODs
+							if (!HLODLayer->IsSpatiallyLoaded() && StaticMesh->bSupportRayTracing)
+							{
+								StaticMeshComponent->bRayTracingFarField = true;
+							}
+
+							// Disable collisions
+							StaticMesh->MarkAsNotHavingNavigationData();
+							if (UBodySetup* BodySetup = StaticMesh->GetBodySetup())
+							{
+								BodySetup->DefaultInstance.SetCollisionProfileName(UCollisionProfile::NoCollision_ProfileName);
+								BodySetup->CollisionTraceFlag = CTF_UseSimpleAsComplex;
+							}
+
+							// Rename owned static mesh
+							StaticMesh->Rename(*MakeUniqueObjectName(StaticMesh->GetOuter(), StaticMesh->GetClass(), *FString::Printf(TEXT("StaticMesh_%s"), *HLODLayer->GetName())).ToString());
+						}
+					}
+				}
+			}
+
+			InHLODActor->SetHLODComponents(HLODComponents);
+		}
+	}
+
+	// Gather stats pertaining to the assets generated during this build
+	GatherOutputStats(InHLODActor);
+
+	TotalTimeMS = FMath::RoundToInt(TotalTimeScope.GetTime() * 1000);
+
+	// Build timings stats
+	InHLODActor->SetStat(FWorldPartitionHLODStats::BuildTimeLoadMilliseconds, LoadTimeMS);
+	InHLODActor->SetStat(FWorldPartitionHLODStats::BuildTimeBuildMilliseconds, BuildTimeMS);
+	InHLODActor->SetStat(FWorldPartitionHLODStats::BuildTimeTotalMilliseconds, TotalTimeMS);
+
+	return NewHLODHash;
+}
+
+#endif
diff --git a/Engine/Plugins/Editor/WorldPartitionHLODUtilities/Source/Private/WorldPartitionHLODUtilities.cpp b/Engine/Plugins/Editor/WorldPartitionHLODUtilities/Source/Private/WorldPartitionHLODUtilities.cpp
deleted file mode 100644
diff --git a/Engine/Plugins/Editor/WorldPartitionHLODUtilities/Source/Private/WorldPartitionHLODUtilitiesModule.cpp b/Engine/Plugins/Editor/WorldPartitionHLODUtilities/Source/Private/WorldPartitionHLODUtilitiesModule.cpp
--- a/Engine/Plugins/Editor/WorldPartitionHLODUtilities/Source/Private/WorldPartitionHLODUtilitiesModule.cpp
+++ b/Engine/Plugins/Editor/WorldPartitionHLODUtilities/Source/Private/WorldPartitionHLODUtilitiesModule.cpp
@@ -2,7 +2,8 @@
 
 #include "WorldPartitionHLODUtilitiesModule.h"
 #include "Modules/ModuleManager.h"
-#include "WorldPartitionHLODUtilities.h"
+
+#include "WorldPartition/HLOD/Utilities/WorldPartitionHLODUtilities.h"
 
 IMPLEMENT_MODULE(FWorldPartitionHLODUtilitiesModule, WorldPartitionHLODUtilities);
 
diff --git a/Engine/Plugins/Editor/WorldPartitionHLODUtilities/Source/Public/HLODBuilderInstancing.h b/Engine/Plugins/Editor/WorldPartitionHLODUtilities/Source/Public/HLODBuilderInstancing.h
deleted file mode 100644
diff --git a/Engine/Plugins/Editor/WorldPartitionHLODUtilities/Source/Public/HLODBuilderMeshApproximate.h b/Engine/Plugins/Editor/WorldPartitionHLODUtilities/Source/Public/HLODBuilderMeshApproximate.h
deleted file mode 100644
diff --git a/Engine/Plugins/Editor/WorldPartitionHLODUtilities/Source/Public/HLODBuilderMeshMerge.h b/Engine/Plugins/Editor/WorldPartitionHLODUtilities/Source/Public/HLODBuilderMeshMerge.h
deleted file mode 100644
diff --git a/Engine/Plugins/Editor/WorldPartitionHLODUtilities/Source/Public/HLODBuilderMeshSimplify.h b/Engine/Plugins/Editor/WorldPartitionHLODUtilities/Source/Public/HLODBuilderMeshSimplify.h
deleted file mode 100644
diff --git a/Engine/Plugins/Editor/WorldPartitionHLODUtilities/Source/Public/WorldPartition/HLOD/Builders/HLODBuilderInstancing.h b/Engine/Plugins/Editor/WorldPartitionHLODUtilities/Source/Public/WorldPartition/HLOD/Builders/HLODBuilderInstancing.h
new file mode 100644
--- /dev/null
+++ b/Engine/Plugins/Editor/WorldPartitionHLODUtilities/Source/Public/WorldPartition/HLOD/Builders/HLODBuilderInstancing.h
@@ -0,0 +1,39 @@
+// Copyright Epic Games, Inc. All Rights Reserved.
+
+#pragma once
+
+#include "Templates/SubclassOf.h"
+#include "WorldPartition/HLOD/HLODBuilder.h"
+#include "HLODBuilderInstancing.generated.h"
+
+
+UCLASS(Blueprintable, Config = Engine, PerObjectConfig)
+class WORLDPARTITIONHLODUTILITIES_API UHLODBuilderInstancingSettings : public UHLODBuilderSettings
+{
+	GENERATED_UCLASS_BODY()
+
+	virtual uint32 GetCRC() const override;
+
+	/**
+	 * If enabled, the components created for the HLODs will not use Nanite.
+	 * Necessary if you want to use the last LOD & the mesh is Nanite enabled, as forced LODs are ignored by Nanite
+	 */
+	UPROPERTY(EditAnywhere, AdvancedDisplay, BlueprintReadWrite, Category = Rendering)
+	bool bDisallowNanite;
+};
+
+
+/**
+ * Build a AWorldPartitionHLOD whose components are ISMC
+ */
+UCLASS(HideDropdown)
+class WORLDPARTITIONHLODUTILITIES_API UHLODBuilderInstancing : public UHLODBuilder
+{
+	 GENERATED_UCLASS_BODY()
+
+public:
+	virtual bool RequiresWarmup() const override { return false; }
+	virtual bool ShouldIgnoreBatchingPolicy() const override { return true; }
+	virtual TSubclassOf<UHLODBuilderSettings> GetSettingsClass() const override;
+	virtual TArray<UActorComponent*> Build(const FHLODBuildContext& InHLODBuildContext, const TArray<UActorComponent*>& InSourceComponents) const override;
+};
diff --git a/Engine/Plugins/Editor/WorldPartitionHLODUtilities/Source/Public/WorldPartition/HLOD/Builders/HLODBuilderMeshApproximate.h b/Engine/Plugins/Editor/WorldPartitionHLODUtilities/Source/Public/WorldPartition/HLOD/Builders/HLODBuilderMeshApproximate.h
new file mode 100644
--- /dev/null
+++ b/Engine/Plugins/Editor/WorldPartitionHLODUtilities/Source/Public/WorldPartition/HLOD/Builders/HLODBuilderMeshApproximate.h
@@ -0,0 +1,43 @@
+// Copyright Epic Games, Inc. All Rights Reserved.
+
+#pragma once
+
+#include "WorldPartition/HLOD/HLODBuilder.h"
+#include "Engine/MeshMerging.h"
+
+#include "HLODBuilderMeshApproximate.generated.h"
+
+class UMaterialInterface;
+
+class UMaterial;
+
+
+UCLASS(Blueprintable, Config = Engine, PerObjectConfig)
+class WORLDPARTITIONHLODUTILITIES_API UHLODBuilderMeshApproximateSettings : public UHLODBuilderSettings
+{
+	GENERATED_UCLASS_BODY()
+
+	virtual uint32 GetCRC() const override;
+
+	/** Mesh approximation settings */
+	UPROPERTY(EditAnywhere, Config, Category = HLOD, meta=(EditInline))
+	FMeshApproximationSettings MeshApproximationSettings;
+
+	/** Material that will be used by the generated HLOD static mesh */
+	UPROPERTY(EditAnywhere, Config, AdvancedDisplay, Category = HLOD, meta = (DisplayName = "HLOD Material"))
+	TSoftObjectPtr<UMaterialInterface> HLODMaterial;
+};
+
+
+/**
+ * Build an approximated mesh using geometry from the provided actors
+ */
+UCLASS(HideDropdown)
+class WORLDPARTITIONHLODUTILITIES_API UHLODBuilderMeshApproximate : public UHLODBuilder
+{
+	GENERATED_UCLASS_BODY()
+
+public:
+	virtual TSubclassOf<UHLODBuilderSettings> GetSettingsClass() const override;
+	virtual TArray<UActorComponent*> Build(const FHLODBuildContext& InHLODBuildContext, const TArray<UActorComponent*>& InSourceComponents) const override;
+};
diff --git a/Engine/Plugins/Editor/WorldPartitionHLODUtilities/Source/Public/WorldPartition/HLOD/Builders/HLODBuilderMeshMerge.h b/Engine/Plugins/Editor/WorldPartitionHLODUtilities/Source/Public/WorldPartition/HLOD/Builders/HLODBuilderMeshMerge.h
new file mode 100644
--- /dev/null
+++ b/Engine/Plugins/Editor/WorldPartitionHLODUtilities/Source/Public/WorldPartition/HLOD/Builders/HLODBuilderMeshMerge.h
@@ -0,0 +1,43 @@
+// Copyright Epic Games, Inc. All Rights Reserved.
+
+#pragma once
+
+#include "WorldPartition/HLOD/HLODBuilder.h"
+#include "Engine/MeshMerging.h"
+
+#include "HLODBuilderMeshMerge.generated.h"
+
+class UMaterialInterface;
+
+class UMaterial;
+
+
+UCLASS(Blueprintable, Config = Engine, PerObjectConfig)
+class WORLDPARTITIONHLODUTILITIES_API UHLODBuilderMeshMergeSettings : public UHLODBuilderSettings
+{
+	GENERATED_UCLASS_BODY()
+
+	virtual uint32 GetCRC() const override;
+
+	/** Merged mesh generation settings */
+	UPROPERTY(EditAnywhere, Config, Category = HLOD)
+	FMeshMergingSettings MeshMergeSettings;
+
+	/** Material that will be used by the generated HLOD static mesh */
+	UPROPERTY(EditAnywhere, Config, AdvancedDisplay, Category = HLOD, meta = (DisplayName = "HLOD Material"))
+	TSoftObjectPtr<UMaterialInterface> HLODMaterial;
+};
+
+
+/**
+ * Build a merged mesh using geometry from the provided actors
+ */
+UCLASS(HideDropdown)
+class WORLDPARTITIONHLODUTILITIES_API UHLODBuilderMeshMerge : public UHLODBuilder
+{
+	GENERATED_UCLASS_BODY()
+
+public:
+	virtual TSubclassOf<UHLODBuilderSettings> GetSettingsClass() const override;
+	virtual TArray<UActorComponent*> Build(const FHLODBuildContext& InHLODBuildContext, const TArray<UActorComponent*>& InSourceComponents) const override;
+};
diff --git a/Engine/Plugins/Editor/WorldPartitionHLODUtilities/Source/Public/WorldPartition/HLOD/Builders/HLODBuilderMeshSimplify.h b/Engine/Plugins/Editor/WorldPartitionHLODUtilities/Source/Public/WorldPartition/HLOD/Builders/HLODBuilderMeshSimplify.h
new file mode 100644
--- /dev/null
+++ b/Engine/Plugins/Editor/WorldPartitionHLODUtilities/Source/Public/WorldPartition/HLOD/Builders/HLODBuilderMeshSimplify.h
@@ -0,0 +1,43 @@
+// Copyright Epic Games, Inc. All Rights Reserved.
+
+#pragma once
+
+#include "WorldPartition/HLOD/HLODBuilder.h"
+#include "Engine/MeshMerging.h"
+
+#include "HLODBuilderMeshSimplify.generated.h"
+
+class UMaterialInterface;
+
+class UMaterial;
+
+
+UCLASS(Blueprintable, Config = Engine, PerObjectConfig)
+class WORLDPARTITIONHLODUTILITIES_API UHLODBuilderMeshSimplifySettings : public UHLODBuilderSettings
+{
+	GENERATED_UCLASS_BODY()
+
+	virtual uint32 GetCRC() const override;
+
+	/** Simplified mesh generation settings */
+	UPROPERTY(EditAnywhere, Config, Category = HLOD)
+	FMeshProxySettings MeshSimplifySettings;
+
+	/** Material that will be used by the generated HLOD static mesh */
+	UPROPERTY(EditAnywhere, Config, AdvancedDisplay, Category = HLOD, meta = (DisplayName = "HLOD Material"))
+	TSoftObjectPtr<UMaterialInterface> HLODMaterial;
+};
+
+
+/**
+ * Build a simplified mesh using geometry from the provided actors
+ */
+UCLASS(HideDropdown)
+class WORLDPARTITIONHLODUTILITIES_API UHLODBuilderMeshSimplify : public UHLODBuilder
+{
+	GENERATED_UCLASS_BODY()
+
+public:
+	virtual TSubclassOf<UHLODBuilderSettings> GetSettingsClass() const override;
+	virtual TArray<UActorComponent*> Build(const FHLODBuildContext& InHLODBuildContext, const TArray<UActorComponent*>& InSourceComponents) const override;
+};
diff --git a/Engine/Plugins/Editor/WorldPartitionHLODUtilities/Source/Public/WorldPartition/HLOD/Utilities/WorldPartitionHLODUtilities.h b/Engine/Plugins/Editor/WorldPartitionHLODUtilities/Source/Public/WorldPartition/HLOD/Utilities/WorldPartitionHLODUtilities.h
new file mode 100644
--- /dev/null
+++ b/Engine/Plugins/Editor/WorldPartitionHLODUtilities/Source/Public/WorldPartition/HLOD/Utilities/WorldPartitionHLODUtilities.h
@@ -0,0 +1,17 @@
+// Copyright Epic Games, Inc. All Rights Reserved.
+
+#pragma once
+
+#include "WorldPartition/HLOD/IWorldPartitionHLODUtilities.h"
+
+/**
+* FWorldPartitionHLODUtilities implementation
+*/
+class FWorldPartitionHLODUtilities : public IWorldPartitionHLODUtilities
+{
+public:
+	virtual TArray<AWorldPartitionHLOD*> CreateHLODActors(FHLODCreationContext& InCreationContext, const FHLODCreationParams& InCreationParams, const TArray<IStreamingGenerationContext::FActorInstance>& InActors) override;
+	virtual uint32 BuildHLOD(AWorldPartitionHLOD* InHLODActor) override;
+	virtual TSubclassOf<UHLODBuilder> GetHLODBuilderClass(const UHLODLayer* InHLODLayer) override;
+	virtual UHLODBuilderSettings* CreateHLODBuilderSettings(UHLODLayer* InHLODLayer) override;
+};
diff --git a/Engine/Plugins/Editor/WorldPartitionHLODUtilities/Source/Public/WorldPartitionHLODUtilities.h b/Engine/Plugins/Editor/WorldPartitionHLODUtilities/Source/Public/WorldPartitionHLODUtilities.h
deleted file mode 100644
diff --git a/Engine/Plugins/Editor/WorldPartitionHLODUtilities/Source/WorldPartitionHLODUtilities.Build.cs b/Engine/Plugins/Editor/WorldPartitionHLODUtilities/Source/WorldPartitionHLODUtilities.Build.cs
--- a/Engine/Plugins/Editor/WorldPartitionHLODUtilities/Source/WorldPartitionHLODUtilities.Build.cs
+++ b/Engine/Plugins/Editor/WorldPartitionHLODUtilities/Source/WorldPartitionHLODUtilities.Build.cs
@@ -1,11 +1,17 @@
 // Copyright Epic Games, Inc. All Rights Reserved.
 using UnrealBuildTool;
+using System.IO;
 
 public class WorldPartitionHLODUtilities : ModuleRules
 {
     public WorldPartitionHLODUtilities(ReadOnlyTargetRules Target) : base(Target)
 	{
-        PublicDependencyModuleNames.AddRange(
+		// Some files were initially right under /Public and were moved to sub directories
+		// Ensure we don't break old-style includes
+		PublicIncludePaths.Add(Path.Combine(ModuleDirectory, "Public/WorldPartition/HLOD/Builders"));
+		PublicIncludePaths.Add(Path.Combine(ModuleDirectory, "Public/WorldPartition/HLOD/Utilities"));
+		
+		PublicDependencyModuleNames.AddRange(
             new string[]
 			{
 				"Core",
