World Partition - HLOD: Added support for providing a custom HLOD actor class to HLOD layers

diff --git a/Engine/Plugins/Editor/WorldPartitionHLODUtilities/Source/Private/WorldPartition/HLOD/Utilities/WorldPartitionHLODUtilities.cpp b/Engine/Plugins/Editor/WorldPartitionHLODUtilities/Source/Private/WorldPartition/HLOD/Utilities/WorldPartitionHLODUtilities.cpp
--- a/Engine/Plugins/Editor/WorldPartitionHLODUtilities/Source/Private/WorldPartition/HLOD/Utilities/WorldPartitionHLODUtilities.cpp
+++ b/Engine/Plugins/Editor/WorldPartitionHLODUtilities/Source/Private/WorldPartition/HLOD/Utilities/WorldPartitionHLODUtilities.cpp
@@ -189,15 +189,23 @@ TArray<AWorldPartitionHLOD*> FWorldPartitionHLODUtilities::CreateHLODActors(FHLO
 		const TSet<FHLODSubActor>& SubActors = Pair.Value;
 		check(!SubActors.IsEmpty());
 
-		auto ComputeCellHash = [](const FString& HLODLayerName, const FGuid CellGuid)
+		auto ComputeHLODActorUniqueHash = [](const UHLODLayer* HLODLayer, const FGuid CellGuid)
 		{
-			uint32 HLODLayerNameHash = FCrc::StrCrc32(*HLODLayerName);
-			uint32 CellGuidHash = GetTypeHash(CellGuid);
-			return HashCombine(HLODLayerNameHash, CellGuidHash);
+			const uint32 HLODLayerNameHash = FCrc::StrCrc32(*HLODLayer->GetName());
+			const uint32 CellGuidHash = GetTypeHash(CellGuid);
+			uint32 HLODActorHash = HashCombine(HLODLayerNameHash, CellGuidHash);
+
+			if (HLODLayer->GetHLODActorClass() != AWorldPartitionHLOD::StaticClass())
+			{
+				const uint32 HLODActorClassHash = FCrc::StrCrc32(*HLODLayer->GetHLODActorClass()->GetPathName());
+				HLODActorHash = HashCombine(HLODActorHash, HLODActorClassHash);
+			}
+
+			return HLODActorHash;
 		};
 
-		uint64 CellHash = ComputeCellHash(HLODLayer->GetName(), InCreationParams.CellGuid);
-		FName HLODActorName = *FString::Printf(TEXT("%s_%016llx"), *HLODLayer->GetName(), CellHash);		
+		uint64 HLODActorHash = ComputeHLODActorUniqueHash(HLODLayer, InCreationParams.CellGuid);
+		FName HLODActorName = *FString::Printf(TEXT("%s_%016llx"), *HLODLayer->GetName(), HLODActorHash);
 
 		AWorldPartitionHLOD* HLODActor = nullptr;
 		FWorldPartitionHandle HLODActorHandle;
@@ -215,7 +223,7 @@ TArray<AWorldPartitionHLOD*> FWorldPartitionHLODUtilities::CreateHLODActors(FHLO
 			FActorSpawnParameters SpawnParams;
 			SpawnParams.Name = HLODActorName;
 			SpawnParams.NameMode = FActorSpawnParameters::ESpawnActorNameMode::Required_Fatal;
-			HLODActor = InCreationParams.WorldPartition->GetWorld()->SpawnActor<AWorldPartitionHLOD>(SpawnParams);
+			HLODActor = InCreationParams.WorldPartition->GetWorld()->SpawnActor<AWorldPartitionHLOD>(HLODLayer->GetHLODActorClass(), SpawnParams);
 
 			check(HLODActor->GetContentBundleGuid() == InCreationParams.ContentBundleGuid);
 
@@ -233,6 +241,7 @@ TArray<AWorldPartitionHLOD*> FWorldPartitionHLODUtilities::CreateHLODActors(FHLO
 #if DO_CHECK
 			check(HLODActor->GetSourceCellGuid() == InCreationParams.CellGuid);
 			check(HLODActor->GetSubActorsHLODLayer() == HLODLayer);
+			check(HLODActor->GetClass() == HLODLayer->GetHLODActorClass());
 #endif
 		}
 
diff --git a/Engine/Source/Runtime/Engine/Private/WorldPartition/HLOD/HLODLayer.cpp b/Engine/Source/Runtime/Engine/Private/WorldPartition/HLOD/HLODLayer.cpp
--- a/Engine/Source/Runtime/Engine/Private/WorldPartition/HLOD/HLODLayer.cpp
+++ b/Engine/Source/Runtime/Engine/Private/WorldPartition/HLOD/HLODLayer.cpp
@@ -28,6 +28,7 @@ UHLODLayer::UHLODLayer(const FObjectInitializer& ObjectInitializer)
 	, bIsSpatiallyLoaded(true)
 	, CellSize(25600)
 	, LoadingRange(51200)
+	, HLODActorClass(AWorldPartitionHLOD::StaticClass())
 {
 }
 
diff --git a/Engine/Source/Runtime/Engine/Public/WorldPartition/HLOD/HLODLayer.h b/Engine/Source/Runtime/Engine/Public/WorldPartition/HLOD/HLODLayer.h
--- a/Engine/Source/Runtime/Engine/Public/WorldPartition/HLOD/HLODLayer.h
+++ b/Engine/Source/Runtime/Engine/Public/WorldPartition/HLOD/HLODLayer.h
@@ -16,6 +16,7 @@
 #include "HLODLayer.generated.h"
 
 class AActor;
+class AWorldPartitionHLOD;
 class UHLODBuilder;
 class UHLODBuilderSettings;
 class UMaterial;
@@ -55,6 +56,7 @@ public:
 	void SetLayerType(EHLODLayerType InLayerType) { LayerType = InLayerType; }
 	const TSubclassOf<UHLODBuilder> GetHLODBuilderClass() const { return HLODBuilderClass; }
 	const UHLODBuilderSettings* GetHLODBuilderSettings() const { return HLODBuilderSettings; }
+	const TSubclassOf<AWorldPartitionHLOD> GetHLODActorClass() const { return HLODActorClass; }
 	const TSubclassOf<UWorldPartitionHLODModifier> GetHLODModifierClass() const { return HLODModifierClass; }
 	FName GetRuntimeGrid(uint32 InHLODLevel) const;
 	int32 GetCellSize() const { return !bIsSpatiallyLoaded ? 0 : CellSize; }
@@ -106,6 +108,10 @@ private:
 	UPROPERTY(EditAnywhere, Config, Category=HLOD, meta = (EditConditionHides, EditCondition = "bIsSpatiallyLoaded"))
 	TSoftObjectPtr<UHLODLayer> ParentLayer;
 
+	/** Specify a custom HLOD Actor class, the default is AWorldPartitionHLOD */
+	UPROPERTY(EditAnywhere, Category = HLOD, AdvancedDisplay, meta = (DisplayName = "HLOD Actor Class"))
+	TSubclassOf<AWorldPartitionHLOD> HLODActorClass;
+
 	/** HLOD Modifier class, to allow changes to the HLOD at runtime */
 	UPROPERTY(EditAnywhere, Config, Category = HLOD, AdvancedDisplay, meta = (DisplayName = "HLOD Modifier Class"))
 	TSubclassOf<UWorldPartitionHLODModifier> HLODModifierClass;
