World Partition - HLOD: Allow actor classes to provide their own set of components to be used for HLOD building

diff --git a/Engine/Plugins/Editor/WorldPartitionHLODUtilities/Source/Private/WorldPartition/HLOD/Utilities/WorldPartitionHLODUtilities.cpp b/Engine/Plugins/Editor/WorldPartitionHLODUtilities/Source/Private/WorldPartition/HLOD/Utilities/WorldPartitionHLODUtilities.cpp
--- a/Engine/Plugins/Editor/WorldPartitionHLODUtilities/Source/Private/WorldPartition/HLOD/Utilities/WorldPartitionHLODUtilities.cpp
+++ b/Engine/Plugins/Editor/WorldPartitionHLODUtilities/Source/Private/WorldPartition/HLOD/Utilities/WorldPartitionHLODUtilities.cpp
@@ -114,7 +114,7 @@ static uint32 GetCRC(const TArray<FHLODSubActor>& InSubActors)
 	return CRC;
 }
 
-static uint32 ComputeHLODHash(AWorldPartitionHLOD* InHLODActor, const TArray<AActor*>& InActors)
+static uint32 ComputeHLODHash(AWorldPartitionHLOD* InHLODActor, const TArray<UActorComponent*>& InSourceComponents)
 {
 	FArchiveCrc32 Ar;
 
@@ -147,7 +147,7 @@ static uint32 ComputeHLODHash(AWorldPartitionHLOD* InHLODActor, const TArray<AAc
 	}
 
 	// Append all components CRCs
-	uint32 HLODComponentsHash = UHLODBuilder::ComputeHLODHash(InActors);
+	uint32 HLODComponentsHash = UHLODBuilder::ComputeHLODHash(InSourceComponents);
 	UE_LOG(LogHLODBuilder, VeryVerbose, TEXT(" - HLOD Source Components = %x"), HLODComponentsHash);
 	Ar << HLODComponentsHash;
 
@@ -616,6 +616,24 @@ void GatherOutputStats(AWorldPartitionHLOD* InHLODActor)
 	InHLODActor->SetStat(FWorldPartitionHLODStats::MemoryDiskSizeBytes, 0); // Clear disk size as it is unknown at this point. It will be assigned at package loading
 }
 
+// Iterate over the source actors and retrieve HLOD relevant components using GetHLODRelevantComponents()
+static TArray<UActorComponent*> GatherHLODRelevantComponents(const TArray<AActor*>& InSourceActors)
+{
+	TSet<UActorComponent*> HLODRelevantComponents;
+
+	for (AActor* Actor : InSourceActors)
+	{
+		if (!Actor || !Actor->IsHLODRelevant())
+		{
+			continue;
+		}
+
+		HLODRelevantComponents.Append(Actor->GetHLODRelevantComponents());
+	}
+
+	return HLODRelevantComponents.Array();
+}
+
 uint32 FWorldPartitionHLODUtilities::BuildHLOD(AWorldPartitionHLOD* InHLODActor)
 {
 	FAutoScopedDurationTimer TotalTimeScope;
@@ -639,8 +657,10 @@ uint32 FWorldPartitionHLODUtilities::BuildHLOD(AWorldPartitionHLOD* InHLODActor)
 		UWorldPartitionLevelStreamingDynamic::UnloadFromEditor(LevelStreaming);
 	};
 
+	TArray<UActorComponent*> HLODRelevantComponents = GatherHLODRelevantComponents(LevelStreaming->GetLoadedLevel()->Actors);
+
 	uint32 OldHLODHash = bIsDirty ? 0 : InHLODActor->GetHLODHash();
-	uint32 NewHLODHash = ComputeHLODHash(InHLODActor, LevelStreaming->GetLoadedLevel()->Actors);
+	uint32 NewHLODHash = ComputeHLODHash(InHLODActor, HLODRelevantComponents);
 
 	if (OldHLODHash == NewHLODHash)
 	{
@@ -690,7 +710,7 @@ uint32 FWorldPartitionHLODUtilities::BuildHLOD(AWorldPartitionHLOD* InHLODActor)
 
 			FHLODBuildContext HLODBuildContext;
 			HLODBuildContext.World = InHLODActor->GetWorld();
-			HLODBuildContext.SourceActors = LevelStreaming->GetLoadedLevel()->Actors;
+			HLODBuildContext.SourceComponents = HLODRelevantComponents;
 			HLODBuildContext.AssetsOuter = InHLODActor->GetPackage();
 			HLODBuildContext.AssetsBaseName = InHLODActor->GetActorLabel();
 			HLODBuildContext.MinVisibleDistance = InHLODActor->GetMinVisibleDistance();
diff --git a/Engine/Source/Runtime/Engine/Classes/GameFramework/Actor.h b/Engine/Source/Runtime/Engine/Classes/GameFramework/Actor.h
--- a/Engine/Source/Runtime/Engine/Classes/GameFramework/Actor.h
+++ b/Engine/Source/Runtime/Engine/Classes/GameFramework/Actor.h
@@ -2280,6 +2280,12 @@ public:
 	/** Indicates whether this actor contributes to the HLOD generation. */
 	virtual bool IsHLODRelevant() const;
 
+	/** Indicates whether this actor can provide HLOD relevant components. */
+	virtual bool HasHLODRelevantComponents() const;
+
+	/** Return the list of components to use when building the HLOD representation of this actor. */
+	virtual TArray<UActorComponent*> GetHLODRelevantComponents() const;
+
 	/** 
 	 * Set LOD Parent component for all of our components, normally associated with an ALODActor. 
 	 * @param InLODParent			This component used to compute visibility when hierarchical LOD is enabled. 
diff --git a/Engine/Source/Runtime/Engine/Private/Actor.cpp b/Engine/Source/Runtime/Engine/Private/Actor.cpp
--- a/Engine/Source/Runtime/Engine/Private/Actor.cpp
+++ b/Engine/Source/Runtime/Engine/Private/Actor.cpp
@@ -5968,9 +5968,23 @@ bool AActor::IsHLODRelevant() const
 		return false;
 	}
 
+	return HasHLODRelevantComponents();
+}
+
+bool AActor::HasHLODRelevantComponents() const
+{
 	return Algo::AnyOf(GetComponents(), [](const UActorComponent* Component) { return Component->IsHLODRelevant(); });
 }
 
+TArray<UActorComponent*> AActor::GetHLODRelevantComponents() const
+{
+	TArray<UActorComponent*> HLODRelevantComponents;
+	auto IsHLODRelevant = [](UActorComponent* Component) { return Component->IsHLODRelevant(); };
+	auto GetComponent = [](UActorComponent* Component) { return Component; };
+	Algo::TransformIf(GetComponents(), HLODRelevantComponents, IsHLODRelevant, GetComponent);
+	return HLODRelevantComponents;
+}
+
 void AActor::SetLODParent(UPrimitiveComponent* InLODParent, float InParentDrawDistance)
 {
 	if (InLODParent && InLODParent->MinDrawDistance != InParentDrawDistance)
diff --git a/Engine/Source/Runtime/Engine/Private/WorldPartition/HLOD/HLODBuilder.cpp b/Engine/Source/Runtime/Engine/Private/WorldPartition/HLOD/HLODBuilder.cpp
--- a/Engine/Source/Runtime/Engine/Private/WorldPartition/HLOD/HLODBuilder.cpp
+++ b/Engine/Source/Runtime/Engine/Private/WorldPartition/HLOD/HLODBuilder.cpp
@@ -56,29 +56,6 @@ bool UHLODBuilder::RequiresWarmup() const
 	return true;
 }
 
-static TArray<UActorComponent*> GatherHLODRelevantComponents(const TArray<AActor*>& InSourceActors)
-{
-	TArray<UActorComponent*> HLODRelevantComponents;
-
-	for (AActor* Actor : InSourceActors)
-	{
-		if (!Actor || !Actor->IsHLODRelevant())
-		{
-			continue;
-		}
-
-		for (UActorComponent* SubComponent : Actor->GetComponents())
-		{
-			if (SubComponent && SubComponent->IsHLODRelevant())
-			{
-				HLODRelevantComponents.Add(SubComponent);
-			}
-		}
-	}
-
-	return HLODRelevantComponents;
-}
-
 uint32 UHLODBuilder::ComputeHLODHash(const UActorComponent* InSourceComponent) const
 {
 	uint32 ComponentCRC = 0;
@@ -128,12 +105,12 @@ uint32 UHLODBuilder::ComputeHLODHash(const UActorComponent* InSourceComponent) c
 	return ComponentCRC;
 }
 
-uint32 UHLODBuilder::ComputeHLODHash(const TArray<AActor*>& InSourceActors)
+uint32 UHLODBuilder::ComputeHLODHash(const TArray<UActorComponent*>& InSourceComponents)
 {
 	// We get the CRC of each component
 	TArray<uint32> ComponentsCRCs;
 
-	for (UActorComponent* SourceComponent : GatherHLODRelevantComponents(InSourceActors))
+	for (UActorComponent* SourceComponent : InSourceComponents)
 	{
 		TSubclassOf<UHLODBuilder> HLODBuilderClass = SourceComponent->GetCustomHLODBuilderClass();
 		if (!HLODBuilderClass)
@@ -276,19 +253,13 @@ static bool ShouldBatchComponent(UActorComponent* ActorComponent)
 
 TArray<UActorComponent*> UHLODBuilder::Build(const FHLODBuildContext& InHLODBuildContext) const
 {
-	TArray<UActorComponent*> HLODRelevantComponents = GatherHLODRelevantComponents(InHLODBuildContext.SourceActors);
-	if (HLODRelevantComponents.IsEmpty())
-	{
-		return {};
-	}
-
 	// Handle components using a batching policy separately
 	TArray<UActorComponent*> InputComponents;
 	TArray<UActorComponent*> ComponentsToBatch;
 
 	if (!ShouldIgnoreBatchingPolicy())
 	{				
-		for (UActorComponent* SourceComponent : HLODRelevantComponents)
+		for (UActorComponent* SourceComponent : InHLODBuildContext.SourceComponents)
 		{
 			if (ShouldBatchComponent(SourceComponent))
 			{
@@ -302,7 +273,7 @@ TArray<UActorComponent*> UHLODBuilder::Build(const FHLODBuildContext& InHLODBuil
 	}
 	else
 	{
-		InputComponents = MoveTemp(HLODRelevantComponents);
+		InputComponents = InHLODBuildContext.SourceComponents;
 	}
 
 	TMap<TSubclassOf<UHLODBuilder>, TArray<UActorComponent*>> HLODBuildersForComponents;
diff --git a/Engine/Source/Runtime/Engine/Public/WorldPartition/HLOD/HLODBuilder.h b/Engine/Source/Runtime/Engine/Public/WorldPartition/HLOD/HLODBuilder.h
--- a/Engine/Source/Runtime/Engine/Public/WorldPartition/HLOD/HLODBuilder.h
+++ b/Engine/Source/Runtime/Engine/Public/WorldPartition/HLOD/HLODBuilder.h
@@ -41,8 +41,8 @@ struct FHLODBuildContext
 	/** World for which HLODs are being built */
 	UWorld*	World;
 
-	/** Actors that will be represented by this HLOD */
-	TArray<AActor*> SourceActors;
+	/** Components that will be represented by this HLOD */
+	TArray<UActorComponent*> SourceComponents;
 
 	/** Outer to use for generated assets */
 	UObject* AssetsOuter;
@@ -80,7 +80,7 @@ public:
 	 */
 	TArray<UActorComponent*> Build(const FHLODBuildContext& InHLODBuildContext) const;
 
-	UE_DEPRECATED(5.2, "Use Build() method that takes a single FHLODBuildContext parameter. SourceActors are now part of the HLOD build context object.")
+	UE_DEPRECATED(5.2, "Use Build() method that takes a single FHLODBuildContext parameter.")
 	TArray<UActorComponent*> Build(const FHLODBuildContext& InHLODBuildContext, const TArray<AActor*>& InSourceActors) const { return Build(InHLODBuildContext); }
 
 	/**
@@ -110,10 +110,10 @@ public:
 	virtual TArray<UActorComponent*> Build(const FHLODBuildContext& InHLODBuildContext, const TArray<UActorComponent*>& InSourceComponents) const PURE_VIRTUAL(UHLODBuilder::Build, return {};);
 
 	/**
-	 * From a set of actors, compute a unique hash from their properties that are relevant to HLOD generation.
+	 * From a set of components, compute a unique hash from their properties that are relevant to HLOD generation.
 	 * Used to detect changes to the source actors and trigger an HLOD rebuild if necessarry.
 	 */
-	static uint32 ComputeHLODHash(const TArray<AActor*>& InSourceActors);
+	static uint32 ComputeHLODHash(const TArray<UActorComponent*>& InSourceComponents);
 
 	/** 
 	 * Get the InstancedStaticMeshComponent subclass that should be used when creating instanced HLODs.
