World Partition - HLOD: When building an HLOD, handle the case where none of its source actors are found anymore
* ULevelStreaming::LoadedLevel can be null

diff --git a/Engine/Plugins/Editor/WorldPartitionHLODUtilities/Source/Private/WorldPartition/HLOD/Utilities/WorldPartitionHLODUtilities.cpp b/Engine/Plugins/Editor/WorldPartitionHLODUtilities/Source/Private/WorldPartition/HLOD/Utilities/WorldPartitionHLODUtilities.cpp
--- a/Engine/Plugins/Editor/WorldPartitionHLODUtilities/Source/Private/WorldPartition/HLOD/Utilities/WorldPartitionHLODUtilities.cpp
+++ b/Engine/Plugins/Editor/WorldPartitionHLODUtilities/Source/Private/WorldPartition/HLOD/Utilities/WorldPartitionHLODUtilities.cpp
@@ -439,43 +439,51 @@ UHLODBuilderSettings* FWorldPartitionHLODUtilities::CreateHLODBuilderSettings(UH
 	return HLODBuilderSettings;
 }
 
-void GatherInputStats(AWorldPartitionHLOD* InHLODActor, const TArray<AActor*> SubActors)
+void GatherInputStats(AWorldPartitionHLOD* InHLODActor, const TArray<UActorComponent*> InHLODRelevantComponents)
 {
 	int64 NumActors = 0;
 	int64 NumTriangles = 0;
 	int64 NumVertices = 0;
 
-	for (AActor* SubActor : SubActors)
+	TSet<AActor*> HLODRelevantActors;
+
+	for (UActorComponent* HLODRelevantComponent : InHLODRelevantComponents)
 	{
-		if (SubActor && SubActor->IsHLODRelevant())
+		bool bAlreadyInSet = false;
+		AActor* SubActor = HLODRelevantActors.FindOrAdd(HLODRelevantComponent->GetOwner(), &bAlreadyInSet);
+
+		if (!bAlreadyInSet)
 		{
 			if (AWorldPartitionHLOD* SubHLODActor = Cast<AWorldPartitionHLOD>(SubActor))
 			{
 				NumActors += SubHLODActor->GetStat(FWorldPartitionHLODStats::InputActorCount);
 				NumTriangles += SubHLODActor->GetStat(FWorldPartitionHLODStats::InputTriangleCount);
-				NumVertices += SubHLODActor->GetStat(FWorldPartitionHLODStats::InputVertexCount);				
+				NumVertices += SubHLODActor->GetStat(FWorldPartitionHLODStats::InputVertexCount);
 			}
 			else
 			{
 				NumActors++;
+			}
+		}
 
-				SubActor->ForEachComponent<UStaticMeshComponent>(false, [&] (UStaticMeshComponent* SMComponent)
-				{
-					const UStaticMesh* StaticMesh = SMComponent->GetStaticMesh();
-					const FStaticMeshRenderData* RenderData = StaticMesh ? StaticMesh->GetRenderData() : nullptr;
-					const bool bHasRenderData = RenderData && !RenderData->LODResources.IsEmpty();
-
-					if (bHasRenderData)
-					{
-						const UInstancedStaticMeshComponent* ISMComponent = Cast<UInstancedStaticMeshComponent>(SMComponent);
-						const int64 LOD0TriCount = RenderData->LODResources[0].GetNumTriangles();
-						const int64 LOD0VtxCount = RenderData->LODResources[0].GetNumVertices();
-						const int64 NumInstances = ISMComponent ? ISMComponent->GetInstanceCount() : 1;
+		if (!SubActor->IsA<AWorldPartitionHLOD>())
+		{
+			if (UStaticMeshComponent* SMComponent = Cast<UStaticMeshComponent>(HLODRelevantComponent))
+			{
+				const UStaticMesh* StaticMesh = SMComponent->GetStaticMesh();
+				const FStaticMeshRenderData* RenderData = StaticMesh ? StaticMesh->GetRenderData() : nullptr;
+				const bool bHasRenderData = RenderData && !RenderData->LODResources.IsEmpty();
 
-						NumTriangles += LOD0TriCount * NumInstances;
-						NumVertices += LOD0VtxCount * NumInstances;
-					}
-				});
+				if (bHasRenderData)
+				{
+					const UInstancedStaticMeshComponent* ISMComponent = Cast<UInstancedStaticMeshComponent>(SMComponent);
+					const int64 LOD0TriCount = RenderData->LODResources[0].GetNumTriangles();
+					const int64 LOD0VtxCount = RenderData->LODResources[0].GetNumVertices();
+					const int64 NumInstances = ISMComponent ? ISMComponent->GetInstanceCount() : 1;
+					
+					NumTriangles += LOD0TriCount * NumInstances;
+					NumVertices += LOD0VtxCount * NumInstances;
+				}
 			}
 		}
 	}
@@ -658,7 +666,12 @@ uint32 FWorldPartitionHLODUtilities::BuildHLOD(AWorldPartitionHLOD* InHLODActor)
 		UWorldPartitionLevelStreamingDynamic::UnloadFromEditor(LevelStreaming);
 	};
 
-	TArray<UActorComponent*> HLODRelevantComponents = GatherHLODRelevantComponents(LevelStreaming->GetLoadedLevel()->Actors);
+	TArray<UActorComponent*> HLODRelevantComponents;
+	
+	if (LevelStreaming->GetLoadedLevel())
+	{
+		HLODRelevantComponents = GatherHLODRelevantComponents(LevelStreaming->GetLoadedLevel()->Actors);
+	}
 
 	uint32 OldHLODHash = bIsDirty ? 0 : InHLODActor->GetHLODHash();
 	uint32 NewHLODHash = ComputeHLODHash(InHLODActor, HLODRelevantComponents);
@@ -695,7 +708,7 @@ uint32 FWorldPartitionHLODUtilities::BuildHLOD(AWorldPartitionHLOD* InHLODActor)
 	}
 
 	// Gather stats from the input to our HLOD build
-	GatherInputStats(InHLODActor, LevelStreaming->GetLoadedLevel()->Actors);
+	GatherInputStats(InHLODActor, HLODRelevantComponents);
 	
 	const UHLODLayer* HLODLayer = InHLODActor->GetSubActorsHLODLayer();
 	TSubclassOf<UHLODBuilder> HLODBuilderClass = GetHLODBuilderClass(HLODLayer);
diff --git a/Engine/Source/Runtime/Engine/Private/WorldPartition/WorldPartitionLevelStreamingDynamic.cpp b/Engine/Source/Runtime/Engine/Private/WorldPartition/WorldPartitionLevelStreamingDynamic.cpp
--- a/Engine/Source/Runtime/Engine/Private/WorldPartition/WorldPartitionLevelStreamingDynamic.cpp
+++ b/Engine/Source/Runtime/Engine/Private/WorldPartition/WorldPartitionLevelStreamingDynamic.cpp
@@ -113,16 +113,19 @@ void UWorldPartitionLevelStreamingDynamic::UnloadFromEditor(UWorldPartitionLevel
 	UWorld* World = InLevelStreaming->GetWorld();
 	check(World->WorldType == EWorldType::Editor);
 
-	ULevel* Level = InLevelStreaming->GetLoadedLevel();
 	InLevelStreaming->SetShouldBeVisibleInEditor(false);
 	InLevelStreaming->SetIsRequestingUnloadAndRemoval(true);
-	World->RemoveLevel(Level);
-	World->FlushLevelStreaming();
 
-	// Destroy the package world and remove it from root
-	UPackage* Package = Level->GetPackage();
-	UWorld* PackageWorld = UWorld::FindWorldInPackage(Package);
-	PackageWorld->DestroyWorld(false);
+	if (ULevel* Level = InLevelStreaming->GetLoadedLevel())
+	{
+		World->RemoveLevel(Level);
+		World->FlushLevelStreaming();
+
+		// Destroy the package world and remove it from root
+		UPackage* Package = Level->GetPackage();
+		UWorld* PackageWorld = UWorld::FindWorldInPackage(Package);
+		PackageWorld->DestroyWorld(false);
+	}
 }
 
 void UWorldPartitionLevelStreamingDynamic::Initialize(UWorld* OuterWorld, const TArray<FWorldPartitionRuntimeCellObjectMapping>& InPackages)
