World Partition - HLOD: Added support for WP HLOD from traditional levels

Introduced a HLOD Source Actors (base) class, with the following subclasses:
* UWorldPartitionHLODSourceActorsFromCell - Actors coming from WP cells (was the previous storage for source actors)
* UWorldPartitionHLODSourceActorsFromLevel - Source is a World asset in which all relevant actors are considered

Added HLOD Provider interface
* When implemented, can be used to create a WP HLOD Actor.

WP HLOD Builder now has the ability to generate WP HLOD actors from non partitioned maps

diff --git a/Engine/Plugins/Editor/WorldPartitionHLODUtilities/Source/Private/WorldPartition/HLOD/Utilities/WorldPartitionHLODUtilities.cpp b/Engine/Plugins/Editor/WorldPartitionHLODUtilities/Source/Private/WorldPartition/HLOD/Utilities/WorldPartitionHLODUtilities.cpp
--- a/Engine/Plugins/Editor/WorldPartitionHLODUtilities/Source/Private/WorldPartition/HLOD/Utilities/WorldPartitionHLODUtilities.cpp
+++ b/Engine/Plugins/Editor/WorldPartitionHLODUtilities/Source/Private/WorldPartition/HLOD/Utilities/WorldPartitionHLODUtilities.cpp
@@ -18,7 +18,7 @@
 #include "WorldPartition/HLOD/HLODLayer.h"
 #include "WorldPartition/HLOD/HLODModifier.h"
 #include "WorldPartition/HLOD/HLODStats.h"
-#include "WorldPartition/WorldPartitionLevelStreamingDynamic.h"
+#include "WorldPartition/HLOD/HLODSourceActorsFromCell.h"
 #include "WorldPartition/ContentBundle/ContentBundleActivationScope.h"
 
 #include "WorldPartition/HLOD/Builders/HLODBuilderInstancing.h"
@@ -27,7 +27,9 @@
 #include "WorldPartition/HLOD/Builders/HLODBuilderMeshApproximate.h"
 
 #include "Engine/Level.h"
+#include "Engine/LevelStreaming.h"
 #include "Engine/StaticMesh.h"
+#include "Engine/World.h"
 #include "Materials/MaterialInstance.h"
 #include "PhysicsEngine/BodySetup.h"
 #include "ProfilingDebugging/ScopedTimers.h"
@@ -37,105 +39,25 @@
 #include "UObject/MetaData.h"
 #include "UObject/GCObjectScopeGuard.h"
 
-
-static UWorldPartitionLevelStreamingDynamic* CreateLevelStreamingFromHLODActor(AWorldPartitionHLOD* InHLODActor, bool& bOutDirty)
-{
-	UPackage::WaitForAsyncFileWrites();
-
-	bOutDirty = false;
-	UWorld* World = InHLODActor->GetWorld();
-	UWorldPartition* WorldPartition = World->GetWorldPartition();
-	check(WorldPartition);
-	
-	const FName LevelStreamingName = FName(*FString::Printf(TEXT("HLODLevelStreaming_%s"), *InHLODActor->GetName()));
-	TArray<FWorldPartitionRuntimeCellObjectMapping> Mappings;
-	Mappings.Reserve(InHLODActor->GetSubActors().Num());
-	Algo::Transform(InHLODActor->GetSubActors(), Mappings, [World, InHLODActor](const FHLODSubActor& SubActor)
-	{
-		return FWorldPartitionRuntimeCellObjectMapping(
-			SubActor.ActorPackage, 
-			SubActor.ActorPath, 
-			SubActor.ContainerID, 
-			SubActor.ContainerTransform, 
-			SubActor.ContainerPackage, 
-			World->GetPackage()->GetFName(), 
-			InHLODActor->GetContentBundleGuid(),
-			SubActor.ActorGuid,
-			false
-		);
-	});
-
-	UWorldPartitionLevelStreamingDynamic* LevelStreaming = UWorldPartitionLevelStreamingDynamic::LoadInEditor(World, LevelStreamingName, Mappings);
-	check(LevelStreaming);
-
-	if (LevelStreaming->GetLoadSucceeded())
-	{
-		// Finish assets compilation
-		FAssetCompilingManager::Get().FinishAllCompilation();
-
-		// Ensure all deferred construction scripts are executed
-		FAssetCompilingManager::Get().ProcessAsyncTasks();
-	}
-	else
-	{
-		bOutDirty = true;
-		UE_LOG(LogHLODBuilder, Warning, TEXT("HLOD actor \"%s\" needs to be rebuilt as it didn't succeed in loading all actors."), *InHLODActor->GetActorLabel());
-	}
-
-	return LevelStreaming;
-}
-
-static uint32 GetCRC(const UHLODLayer* InHLODLayer)
-{
-	UHLODLayer& HLODLayer= *const_cast<UHLODLayer*>(InHLODLayer);
-
-	uint32 CRC;
-
-	CRC = GetTypeHash(HLODLayer.GetLayerType());
-	UE_LOG(LogHLODBuilder, VeryVerbose, TEXT(" - LayerType = %d"), CRC);
-
-	CRC = HashCombine(HLODLayer.GetHLODBuilderSettings()->GetCRC(), CRC);
-	UE_LOG(LogHLODBuilder, VeryVerbose, TEXT(" - HLODBuilderSettings = %d"), CRC);
-
-	CRC = HashCombine(HLODLayer.GetCellSize(), CRC);
-	UE_LOG(LogHLODBuilder, VeryVerbose, TEXT(" - CellSize = %d"), CRC);
-
-	return CRC;
-}
-
-static uint32 GetCRC(const TArray<FHLODSubActor>& InSubActors)
-{
-	uint32 CRC = 0;
-
-	for (const FHLODSubActor& HLODSubActor : InSubActors)
-	{
-		CRC = HashCombine(GetTypeHash(HLODSubActor), CRC);
-	}
-
-	return CRC;
-}
-
 static uint32 ComputeHLODHash(AWorldPartitionHLOD* InHLODActor, const TArray<UActorComponent*>& InSourceComponents)
 {
 	FArchiveCrc32 Ar;
 
 	// Base key, changing this will force a rebuild of all HLODs
-	FString HLODBaseKey = "BFCB19852F0E4DF3BA4CAE6B85C3078D";
+	FString HLODBaseKey = "3B2067A817E140B1926BBCB3015E817A";
 	Ar << HLODBaseKey;
 
-	// Sub Actors
-	uint32 SubActorsHash = GetCRC(InHLODActor->GetSubActors());
-	UE_LOG(LogHLODBuilder, VeryVerbose, TEXT(" - Sub Actors (%d actors) = %x"), InHLODActor->GetSubActors().Num(), SubActorsHash);
-	Ar << SubActorsHash;
-
-	// HLOD Layer
-	uint32 HLODLayerHash = GetCRC(InHLODActor->GetSubActorsHLODLayer());
-	UE_LOG(LogHLODBuilder, VeryVerbose, TEXT(" - HLOD Layer (%s) = %x"), *InHLODActor->GetSubActorsHLODLayer()->GetName(), HLODLayerHash);
-	Ar << HLODLayerHash;
-
+	// HLOD Source Actors
+	if (ensure(InHLODActor->GetSourceActors()))
+	{
+		uint32 SourceActorsHash = InHLODActor->GetSourceActors()->GetHLODHash();
+		UE_LOG(LogHLODHash, VeryVerbose, TEXT(" - SourceActorsHash = %x"), SourceActorsHash);
+		Ar << SourceActorsHash;
+	}
+	
 	// Min Visible Distance
 	uint32 HLODMinVisibleDistanceHash = GetTypeHash(InHLODActor->GetMinVisibleDistance());
-	UE_LOG(LogHLODBuilder, VeryVerbose, TEXT(" - HLOD Min Visible Distance (%.02f) = %x"), InHLODActor->GetMinVisibleDistance(), HLODMinVisibleDistanceHash);
+	UE_LOG(LogHLODHash, VeryVerbose, TEXT(" - HLOD Min Visible Distance (%.02f) = %x"), InHLODActor->GetMinVisibleDistance(), HLODMinVisibleDistanceHash);
 	Ar << HLODMinVisibleDistanceHash;
 
 	// ISM Component Class
@@ -143,13 +65,13 @@ static uint32 ComputeHLODHash(AWorldPartitionHLOD* InHLODActor, const TArray<UAc
 	if (HLODISMComponentClass != UInstancedStaticMeshComponent::StaticClass())
 	{
 		uint32 HLODISMComponentClassHash = GetTypeHash(HLODISMComponentClass);
-		UE_LOG(LogHLODBuilder, VeryVerbose, TEXT(" - HLOD ISM Component Class (%s) = %x"), *HLODISMComponentClass->GetName(), HLODISMComponentClassHash);
+		UE_LOG(LogHLODHash, VeryVerbose, TEXT(" - HLOD ISM Component Class (%s) = %x"), *HLODISMComponentClass->GetName(), HLODISMComponentClassHash);
 		Ar << HLODISMComponentClassHash;
 	}
 
 	// Append all components CRCs
 	uint32 HLODComponentsHash = UHLODBuilder::ComputeHLODHash(InSourceComponents);
-	UE_LOG(LogHLODBuilder, VeryVerbose, TEXT(" - HLOD Source Components = %x"), HLODComponentsHash);
+	UE_LOG(LogHLODHash, VeryVerbose, TEXT(" - HLOD Source Components = %x"), HLODComponentsHash);
 	Ar << HLODComponentsHash;
 
 	return Ar.GetCrc();
@@ -247,7 +169,6 @@ TArray<AWorldPartitionHLOD*> FWorldPartitionHLODUtilities::CreateHLODActors(FHLO
 			check(HLODActor->GetContentBundleGuid() == InCreationParams.ContentBundleGuid);
 
 			HLODActor->SetSourceCellGuid(InCreationParams.CellGuid);
-			HLODActor->SetSubActorsHLODLayer(HLODLayer);
 
 			// Make sure the generated HLOD actor has the same data layers as the source actors
 			for (const UDataLayerInstance* DataLayerInstance : InCreationParams.DataLayerInstances)
@@ -257,21 +178,29 @@ TArray<AWorldPartitionHLOD*> FWorldPartitionHLODUtilities::CreateHLODActors(FHLO
 		}
 		else
 		{
-#if DO_CHECK
 			check(HLODActor->GetSourceCellGuid() == InCreationParams.CellGuid);
-			check(HLODActor->GetSubActorsHLODLayer() == HLODLayer);
 			check(HLODActor->GetClass() == HLODLayer->GetHLODActorClass());
-#endif
 		}
 
 		bool bIsDirty = false;
 
+		// Source actors object
+		UWorldPartitionHLODSourceActorsFromCell* HLODSourceActors = Cast<UWorldPartitionHLODSourceActorsFromCell>(HLODActor->GetSourceActors());
+		if (!HLODSourceActors)
+		{
+			HLODSourceActors = NewObject<UWorldPartitionHLODSourceActorsFromCell>(HLODActor);
+			HLODSourceActors->SetHLODLayer(HLODLayer);
+			HLODActor->SetSourceActors(HLODSourceActors);
+			bIsDirty = true;
+		}
+		check(HLODSourceActors->GetHLODLayer() == HLODLayer);
+
 		// Sub actors
 		{
-			bool bSubActorsChanged = HLODActor->GetSubActors().Num() != SubActors.Num();
+			bool bSubActorsChanged = HLODSourceActors->GetActors().Num() != SubActors.Num();
 			if (!bSubActorsChanged)
 			{
-				TArray<FHLODSubActor> A = HLODActor->GetSubActors();
+				TArray<FHLODSubActor> A = HLODSourceActors->GetActors();
 				TArray<FHLODSubActor> B = SubActors.Array();
 				A.Sort();
 				B.Sort();
@@ -280,7 +209,7 @@ TArray<AWorldPartitionHLOD*> FWorldPartitionHLODUtilities::CreateHLODActors(FHLO
 
 			if (bSubActorsChanged)
 			{
-				HLODActor->SetSubActors(SubActors.Array());
+				HLODSourceActors->SetActors(SubActors.Array());
 				bIsDirty = true;
 			}
 		}
@@ -643,6 +572,42 @@ static TArray<UActorComponent*> GatherHLODRelevantComponents(const TArray<AActor
 	return HLODRelevantComponents.Array();
 }
 
+ULevelStreaming* LoadSourceActors(AWorldPartitionHLOD* InHLODActor, bool& bOutIsDirty)
+{
+	ULevelStreaming* LevelStreaming = InHLODActor->GetSourceActors()->LoadSourceActors(bOutIsDirty);
+	UE_CLOG(bOutIsDirty, LogHLODBuilder, Warning, TEXT("HLOD actor \"%s\" needs to be rebuilt as it didn't succeed in loading all actors."), *InHLODActor->GetActorLabel());
+
+	if (LevelStreaming->GetLoadedLevel())
+	{
+		// Finish assets compilation
+		FAssetCompilingManager::Get().FinishAllCompilation();
+
+		// Ensure all deferred construction scripts are executed
+		FAssetCompilingManager::Get().ProcessAsyncTasks();
+	}
+
+	return LevelStreaming;
+}
+
+void UnloadSourceActors(ULevelStreaming* InLevelStreaming)
+{
+	UWorld* World = InLevelStreaming->GetWorld();
+
+	InLevelStreaming->SetShouldBeVisibleInEditor(false);
+	InLevelStreaming->SetIsRequestingUnloadAndRemoval(true);
+
+	if (ULevel* Level = InLevelStreaming->GetLoadedLevel())
+	{
+		World->RemoveLevel(Level);
+		World->FlushLevelStreaming();
+
+		// Destroy the package world and remove it from root
+		UPackage* Package = Level->GetPackage();
+		UWorld* PackageWorld = UWorld::FindWorldInPackage(Package);
+		PackageWorld->DestroyWorld(false);
+	}
+}
+
 uint32 FWorldPartitionHLODUtilities::BuildHLOD(AWorldPartitionHLOD* InHLODActor)
 {
 	FAutoScopedDurationTimer TotalTimeScope;
@@ -654,16 +619,18 @@ uint32 FWorldPartitionHLODUtilities::BuildHLOD(AWorldPartitionHLOD* InHLODActor)
 	
 	// Load actors relevant to HLODs
 	bool bIsDirty = false;
-	UWorldPartitionLevelStreamingDynamic* LevelStreaming = nullptr;
+	ULevelStreaming* LevelStreaming = nullptr;
 	{
 		FAutoScopedDurationTimer LoadTimeScope;
-		LevelStreaming = CreateLevelStreamingFromHLODActor(InHLODActor, bIsDirty);
+
+		LevelStreaming = LoadSourceActors(InHLODActor, bIsDirty);
+
 		LoadTimeMS = FMath::RoundToInt(LoadTimeScope.GetTime() * 1000);
 	}
 
 	ON_SCOPE_EXIT
 	{
-		UWorldPartitionLevelStreamingDynamic::UnloadFromEditor(LevelStreaming);
+		UnloadSourceActors(LevelStreaming);
 	};
 
 	TArray<UActorComponent*> HLODRelevantComponents;
@@ -691,7 +658,7 @@ uint32 FWorldPartitionHLODUtilities::BuildHLOD(AWorldPartitionHLOD* InHLODActor)
 	    TArray<UObject*> ObjectsToRename;
 	    ForEachObjectWithOuter(InHLODActor->GetPackage(), [&ObjectsToRename](UObject* Obj)
 	    {
-		    if (!Obj->IsA<AActor>() && !Obj->IsA<UMetaData>())
+		    if (!Obj->IsA<AActor>() && !Obj->IsA<UMetaData>() && !Obj->IsA<UWorld>())
 		    {
 			    ObjectsToRename.Add(Obj);
 		    }
@@ -710,7 +677,7 @@ uint32 FWorldPartitionHLODUtilities::BuildHLOD(AWorldPartitionHLOD* InHLODActor)
 	// Gather stats from the input to our HLOD build
 	GatherInputStats(InHLODActor, HLODRelevantComponents);
 	
-	const UHLODLayer* HLODLayer = InHLODActor->GetSubActorsHLODLayer();
+	const UHLODLayer* HLODLayer = InHLODActor->GetSourceActors()->GetHLODLayer();
 	TSubclassOf<UHLODBuilder> HLODBuilderClass = GetHLODBuilderClass(HLODLayer);
 
 	if (HLODBuilderClass)
diff --git a/Engine/Source/Editor/UnrealEd/Private/WorldPartition/WorldPartitionHLODsBuilder.cpp b/Engine/Source/Editor/UnrealEd/Private/WorldPartition/WorldPartitionHLODsBuilder.cpp
--- a/Engine/Source/Editor/UnrealEd/Private/WorldPartition/WorldPartitionHLODsBuilder.cpp
+++ b/Engine/Source/Editor/UnrealEd/Private/WorldPartition/WorldPartitionHLODsBuilder.cpp
@@ -29,6 +29,7 @@
 #include "WorldPartition/HLOD/HLODActor.h"
 #include "WorldPartition/HLOD/HLODActorDesc.h"
 #include "WorldPartition/HLOD/HLODLayer.h"
+#include "WorldPartition/HLOD/HLODProviderInterface.h"
 #include "WorldPartition/HLOD/HLODSubsystem.h"
 
 DEFINE_LOG_CATEGORY_STATIC(LogWorldPartitionHLODsBuilder, All, All);
@@ -249,11 +250,11 @@ bool UWorldPartitionHLODsBuilder::ValidateParams() const
 	return true;
 }
 
-bool UWorldPartitionHLODsBuilder::PreWorldInitialization(UWorld* World, FPackageSourceControlHelper& PackageHelper)
+bool UWorldPartitionHLODsBuilder::PreWorldInitialization(UWorld* InWorld, FPackageSourceControlHelper& PackageHelper)
 {
 	if (bDistributedBuild)
 	{
-		uint32 WorldPackageHash = GetTypeHash(World->GetPackage()->GetFullName());
+		uint32 WorldPackageHash = GetTypeHash(InWorld->GetPackage()->GetFullName());
 		DistributedBuildWorkingDir = FString::Printf(TEXT("%s/%s/%08x"), *FPaths::RootDir(), *DistributedBuildWorkingDirName, WorldPackageHash);
 		DistributedBuildManifest = DistributedBuildWorkingDir / DistributedBuildManifestName;
 
@@ -282,11 +283,11 @@ bool UWorldPartitionHLODsBuilder::PreWorldInitialization(UWorld* World, FPackage
 	return bRet;
 }
 
-bool UWorldPartitionHLODsBuilder::RunInternal(UWorld* World, const FCellInfo& InCellInfo, FPackageSourceControlHelper& PackageHelper)
+bool UWorldPartitionHLODsBuilder::RunInternal(UWorld* InWorld, const FCellInfo& InCellInfo, FPackageSourceControlHelper& PackageHelper)
 {
+	World = InWorld;
 	WorldPartition = World->GetWorldPartition();
-	check(WorldPartition);
-
+	
 	// Allows HLOD Streaming levels to be GCed properly
 	FLevelStreamingGCHelper::EnableForCommandlet();
 
@@ -327,6 +328,12 @@ bool UWorldPartitionHLODsBuilder::RunInternal(UWorld* World, const FCellInfo& In
 
 bool UWorldPartitionHLODsBuilder::SetupHLODActors()
 {
+	// No setup needed for non partitioned worlds
+	if (!WorldPartition)
+	{
+		return true;
+	}
+
 	auto ActorFolderAddedDelegateHandle = GEngine->OnActorFolderAdded().AddLambda([this](UActorFolder* InActorFolder)
 	{
 		UPackage* ActorFolderPackage = InActorFolder->GetPackage();
@@ -439,63 +446,101 @@ bool UWorldPartitionHLODsBuilder::SetupHLODActors()
 
 bool UWorldPartitionHLODsBuilder::BuildHLODActors()
 {
-	TArray<FGuid> HLODActorsToBuild;
-	if (!GetHLODActorsToBuild(HLODActorsToBuild))
+	auto SaveHLODActor = [this](AWorldPartitionHLOD* HLODActor)
 	{
-		return false;
-	}
+		UPackage* ActorPackage = HLODActor->GetPackage();
+		if (ActorPackage->IsDirty())
+		{
+			UE_LOG(LogWorldPartitionHLODsBuilder, Display, TEXT("HLOD actor %s was modified, saving..."), *HLODActor->GetActorLabel());
 
-	if (!ValidateWorkload(HLODActorsToBuild))
-	{
-		return false;
-	}
+			bool bSaved = SourceControlHelper->Save(ActorPackage);
+			if (!bSaved)
+			{
+				UE_LOG(LogWorldPartitionHLODsBuilder, Error, TEXT("Failed to save %s, exiting..."), *USourceControlHelpers::PackageFilename(ActorPackage));
+				return false;
+			}
+		}
 
-	UE_LOG(LogWorldPartitionHLODsBuilder, Display, TEXT("#### Building %d HLOD actors ####"), HLODActorsToBuild.Num());
-	if (bResumeBuild)
-	{
-		UE_LOG(LogWorldPartitionHLODsBuilder, Display, TEXT("#### Resuming build at %d ####"), ResumeBuildIndex);
-	}
+		return true;
+	};
 
-	for (int32 CurrentActor = ResumeBuildIndex; CurrentActor < HLODActorsToBuild.Num(); ++CurrentActor)
+	if (WorldPartition)
 	{
-		TRACE_BOOKMARK(TEXT("BuildHLOD Start - %d"), CurrentActor);
+		TArray<FGuid> HLODActorsToBuild;
+		if (!GetHLODActorsToBuild(HLODActorsToBuild))
+		{
+			return false;
+		}
 
-		const FGuid& HLODActorGuid = HLODActorsToBuild[CurrentActor];
+		if (!ValidateWorkload(HLODActorsToBuild))
+		{
+			return false;
+		}
 
-		FWorldPartitionReference ActorRef(WorldPartition, HLODActorGuid);
-		FWorldPartitionActorDesc* ActorDesc = ActorRef.Get();
+		UE_LOG(LogWorldPartitionHLODsBuilder, Display, TEXT("#### Building %d HLOD actors ####"), HLODActorsToBuild.Num());
+		if (bResumeBuild)
+		{
+			UE_LOG(LogWorldPartitionHLODsBuilder, Display, TEXT("#### Resuming build at %d ####"), ResumeBuildIndex);
+		}
 
-		AWorldPartitionHLOD* HLODActor = CastChecked<AWorldPartitionHLOD>(ActorDesc->GetActor());
+		for (int32 CurrentActor = ResumeBuildIndex; CurrentActor < HLODActorsToBuild.Num(); ++CurrentActor)
+		{
+			TRACE_BOOKMARK(TEXT("BuildHLOD Start - %d"), CurrentActor);
 
-		UE_LOG(LogWorldPartitionHLODsBuilder, Display, TEXT("[%d / %d] Building HLOD actor %s..."), CurrentActor + 1, HLODActorsToBuild.Num(), *HLODActor->GetActorLabel());
+			const FGuid& HLODActorGuid = HLODActorsToBuild[CurrentActor];
 
-		// Simulate an engine tick to make sure engine & render resources that are queued for deletion are processed.
-		FWorldPartitionHelpers::FakeEngineTick(WorldPartition->GetWorld());
+			FWorldPartitionReference ActorRef(WorldPartition, HLODActorGuid);
+			FWorldPartitionActorDesc* ActorDesc = ActorRef.Get();
 
-		HLODActor->BuildHLOD(bForceBuild);
+			AWorldPartitionHLOD* HLODActor = CastChecked<AWorldPartitionHLOD>(ActorDesc->GetActor());
 
-		UPackage* ActorPackage = HLODActor->GetPackage();
-		if (ActorPackage->IsDirty())
-		{
-			UE_LOG(LogWorldPartitionHLODsBuilder, Display, TEXT("HLOD actor %s was modified, saving..."), *HLODActor->GetActorLabel());
+			UE_LOG(LogWorldPartitionHLODsBuilder, Display, TEXT("[%d / %d] Building HLOD actor %s..."), CurrentActor + 1, HLODActorsToBuild.Num(), *HLODActor->GetActorLabel());
 
-			bool bSaved = SourceControlHelper->Save(ActorPackage);
+			// Simulate an engine tick to make sure engine & render resources that are queued for deletion are processed.
+			FWorldPartitionHelpers::FakeEngineTick(WorldPartition->GetWorld());
+
+			HLODActor->BuildHLOD(bForceBuild);
+
+			bool bSaved = SaveHLODActor(HLODActor);
 			if (!bSaved)
 			{
-				UE_LOG(LogWorldPartitionHLODsBuilder, Error, TEXT("Failed to save %s, exiting..."), *USourceControlHelpers::PackageFilename(ActorPackage));
 				return false;
 			}
+
+			if (FWorldPartitionHelpers::ShouldCollectGarbage())
+			{
+				FWorldPartitionHelpers::DoCollectGarbage();
+			}
+
+			TRACE_BOOKMARK(TEXT("BuildHLOD End - %d"), CurrentActor);
 		}
 
-		if (FWorldPartitionHelpers::ShouldCollectGarbage())
+		UE_LOG(LogWorldPartitionHLODsBuilder, Display, TEXT("#### Built %d HLOD actors ####"), HLODActorsToBuild.Num());
+	}
+	else
+	{
+		uint32 NumHLODActors = 0;
+		for (TActorIterator<AActor> ActorIt(World); ActorIt; ++ActorIt)
 		{
-			FWorldPartitionHelpers::DoCollectGarbage();
-		}
+			if (IWorldPartitionHLODProvider* HLODProvider = Cast<IWorldPartitionHLODProvider>(*ActorIt))
+			{
+				AWorldPartitionHLOD* HLODActor = HLODProvider->CreateHLODActor();
+				if (HLODActor)
+				{
+					HLODActor->BuildHLOD(bForceBuild);
 
-		TRACE_BOOKMARK(TEXT("BuildHLOD End - %d"), CurrentActor);
+					bool bSaved = SaveHLODActor(HLODActor);
+					if (!bSaved)
+					{
+						return false;
+					}
+					NumHLODActors++;
+				}
+			}
+		}
+		UE_LOG(LogWorldPartitionHLODsBuilder, Display, TEXT("#### Built %d HLOD actor ####"), NumHLODActors);
 	}
 
-	UE_LOG(LogWorldPartitionHLODsBuilder, Display, TEXT("#### Built %d HLOD actors ####"), HLODActorsToBuild.Num());
 
 	// Move modified files to the temporary working dir, to be submitted later in the final "submit" pass, from a single machine.
 	if (IsDistributedBuild())
@@ -632,6 +677,8 @@ bool UWorldPartitionHLODsBuilder::GetHLODActorsToBuild(TArray<FGuid>& HLODActors
 
 TArray<TArray<FGuid>> UWorldPartitionHLODsBuilder::GetHLODWorkloads(int32 NumWorkloads) const
 {
+	check(WorldPartition);
+
 	// Build a mapping of 1 HLOD[Level] -> N HLOD[Level - 1]
 	TMap<FGuid, TArray<FGuid>>	HLODParenting;
 	for (FActorDescContainerCollection::TIterator<AWorldPartitionHLOD> HLODIterator(WorldPartition); HLODIterator; ++HLODIterator)
@@ -648,19 +695,7 @@ TArray<TArray<FGuid>> UWorldPartitionHLODsBuilder::GetHLODWorkloads(int32 NumWor
 			continue;
 		}
 
-		TArray<FGuid>& ChildHLODs = HLODParenting.Add(HLODIterator->GetGuid());
-
-		for (const FHLODSubActorDesc& SubActor : HLODIterator->GetSubActors())
-		{
-			if (SubActor.ContainerID.IsMainContainer())
-			{
-				FWorldPartitionActorDesc* SubActorDesc = WorldPartition->GetActorDesc(SubActor.ActorGuid);
-				if (SubActorDesc && SubActorDesc->GetActorNativeClass()->IsChildOf<AWorldPartitionHLOD>())
-				{
-					ChildHLODs.Add(SubActor.ActorGuid);
-				}
-			}
-		}
+		HLODParenting.Add(HLODIterator->GetGuid(), HLODIterator->GetChildHLODActors());
 	}
 
 	// All child HLODs must be built before their parent HLOD
@@ -722,8 +757,10 @@ TArray<TArray<FGuid>> UWorldPartitionHLODsBuilder::GetHLODWorkloads(int32 NumWor
 	return Workloads;
 }
 
-bool UWorldPartitionHLODsBuilder::ValidateWorkload(const TArray<FGuid>&Workload) const
+bool UWorldPartitionHLODsBuilder::ValidateWorkload(const TArray<FGuid>& Workload) const
 {
+	check(WorldPartition);
+
 	TSet<FGuid> ProcessedHLOD;
 	ProcessedHLOD.Reserve(Workload.Num());
 
@@ -745,23 +782,12 @@ bool UWorldPartitionHLODsBuilder::ValidateWorkload(const TArray<FGuid>&Workload)
 
 		const FHLODActorDesc* HLODActorDesc = static_cast<const FHLODActorDesc*>(ActorDesc);
 
-		for (const FHLODSubActorDesc& SubActor : HLODActorDesc->GetSubActors())
+		for (const FGuid& ChildHLODActorGuid : HLODActorDesc->GetChildHLODActors())
 		{
-			if (SubActor.ContainerID.IsMainContainer())
+			if (!ProcessedHLOD.Contains(ChildHLODActorGuid))
 			{
-				FWorldPartitionActorDesc* SubActorDesc = WorldPartition->GetActorDesc(SubActor.ActorGuid);
-
-				// Invalid sub actor guid found, this is unexpected when running distributed builds as the build step is always preceeded by the setup step.
-				check(SubActorDesc || !bDistributedBuild);
-
-				if (SubActorDesc && SubActorDesc->GetActorNativeClass()->IsChildOf<AWorldPartitionHLOD>())
-				{
-					if (!ProcessedHLOD.Contains(SubActor.ActorGuid))
-					{
-						UE_LOG(LogWorldPartitionHLODsBuilder, Error, TEXT("Child HLOD actor missing or out of order in HLOD workload, exiting..."));
-						return false;
-					}
-				}
+				UE_LOG(LogWorldPartitionHLODsBuilder, Error, TEXT("Child HLOD actor missing or out of order in HLOD workload, exiting..."));
+				return false;
 			}
 		}
 
@@ -773,6 +799,8 @@ bool UWorldPartitionHLODsBuilder::ValidateWorkload(const TArray<FGuid>&Workload)
 
 bool UWorldPartitionHLODsBuilder::GenerateBuildManifest(TMap<FString, int32>& FilesToBuilderMap) const
 {
+	check(WorldPartition);
+
 	TArray<TArray<FGuid>> BuildersWorkload = GetHLODWorkloads(BuilderCount);
 
 	FConfigFile ConfigFile;
diff --git a/Engine/Source/Editor/UnrealEd/Public/WorldPartition/WorldPartitionHLODsBuilder.h b/Engine/Source/Editor/UnrealEd/Public/WorldPartition/WorldPartitionHLODsBuilder.h
--- a/Engine/Source/Editor/UnrealEd/Public/WorldPartition/WorldPartitionHLODsBuilder.h
+++ b/Engine/Source/Editor/UnrealEd/Public/WorldPartition/WorldPartitionHLODsBuilder.h
@@ -6,6 +6,12 @@
 #include "WorldPartition/WorldPartitionBuilder.h"
 #include "WorldPartitionHLODsBuilder.generated.h"
 
+
+class UWorld;
+class UWorldPartition;
+class FSourceControlHelper;
+
+
 struct FHLODModifiedFiles
 {
 	enum EFileOperation
@@ -78,8 +84,10 @@ public:
 	virtual bool RequiresCommandletRendering() const override;
 	virtual ELoadingMode GetLoadingMode() const override { return ELoadingMode::Custom; }
 	virtual bool PreWorldInitialization(UWorld* World, FPackageSourceControlHelper& PackageHelper) override;
+
 protected:
 	virtual bool RunInternal(UWorld* World, const FCellInfo& InCellInfo, FPackageSourceControlHelper& PackageHelper) override;
+	virtual bool CanProcessNonPartitionedWorlds() const override { return true; }
 	// UWorldPartitionBuilder interface end
 
 	bool IsDistributedBuild() const { return bDistributedBuild; }
@@ -106,8 +114,9 @@ protected:
 	bool AddBuildProducts(const TArray<FString>& BuildProducts) const;
 
 private:
-	class UWorldPartition* WorldPartition;
-	class FSourceControlHelper* SourceControlHelper;
+	UWorld* World;
+	UWorldPartition* WorldPartition;
+	FSourceControlHelper* SourceControlHelper;
 
 	// Options --
 	EHLODBuildStep BuildOptions;
diff --git a/Engine/Source/Runtime/Core/Public/UObject/FortniteMainBranchObjectVersion.h b/Engine/Source/Runtime/Core/Public/UObject/FortniteMainBranchObjectVersion.h
--- a/Engine/Source/Runtime/Core/Public/UObject/FortniteMainBranchObjectVersion.h
+++ b/Engine/Source/Runtime/Core/Public/UObject/FortniteMainBranchObjectVersion.h
@@ -328,6 +328,9 @@ struct CORE_API FFortniteMainBranchObjectVersion
 		// Add customizable values for several UWorldPartitionRuntimeSpatialHash cvars
 		WorldPartitionRuntimeSpatialHashCVarOverrides,
 
+		// WorldPartition HLOD now contains a source actors object
+		WorldPartitionHLODSourceActorsRefactor,
+
 		// -----<new versions can be added above this line>-------------------------------------------------
 		VersionPlusOne,
 		LatestVersion = VersionPlusOne - 1
diff --git a/Engine/Source/Runtime/Engine/Private/WorldPartition/HLOD/HLODActor.cpp b/Engine/Source/Runtime/Engine/Private/WorldPartition/HLOD/HLODActor.cpp
--- a/Engine/Source/Runtime/Engine/Private/WorldPartition/HLOD/HLODActor.cpp
+++ b/Engine/Source/Runtime/Engine/Private/WorldPartition/HLOD/HLODActor.cpp
@@ -15,17 +15,23 @@
 #include "WorldPartition/WorldPartitionHelpers.h"
 
 #if WITH_EDITOR
+#include "Editor.h"
+#include "Modules/ModuleManager.h"
 #include "Misc/ArchiveMD5.h"
+#include "Subsystems/AssetEditorSubsystem.h"
 #include "UObject/FortniteMainBranchObjectVersion.h"
 #include "WorldPartition/HLOD/HLODActorDesc.h"
 #include "WorldPartition/HLOD/HLODLayer.h"
 #include "WorldPartition/HLOD/HLODStats.h"
 #include "WorldPartition/HLOD/IWorldPartitionHLODUtilities.h"
 #include "WorldPartition/HLOD/IWorldPartitionHLODUtilitiesModule.h"
-
-#include "Modules/ModuleManager.h"
+#include "WorldPartition/HLOD/HLODSourceActorsFromCell.h"
 #endif
 
+
+DEFINE_LOG_CATEGORY(LogHLODHash);
+
+
 static int32 GWorldPartitionHLODForceDisableShadows = 0;
 static FAutoConsoleVariableRef CVarWorldPartitionHLODForceDisableShadows(
 	TEXT("wp.Runtime.HLOD.ForceDisableShadows"),
@@ -33,27 +39,6 @@ static FAutoConsoleVariableRef CVarWorldPartitionHLODForceDisableShadows(
 	TEXT("Force disable CastShadow flag on World Partition HLOD actors"),
 	ECVF_Scalability);
 
-#if WITH_EDITORONLY_DATA
-FArchive& operator<<(FArchive& Ar, FHLODSubActorDesc& SubActor)
-{
-	Ar.UsingCustomVersion(FFortniteMainBranchObjectVersion::GUID);
-
-	Ar << SubActor.ActorGuid;
-
-	if (Ar.CustomVer(FFortniteMainBranchObjectVersion::GUID) < FFortniteMainBranchObjectVersion::WorldPartitionFActorContainerIDu64ToGuid)
-	{
-		uint64 ContainerID;
-		Ar << ContainerID;
-	}
-	else
-	{
-		Ar << SubActor.ContainerID;
-	}
-
-	return Ar;
-}
-#endif
-
 AWorldPartitionHLOD::AWorldPartitionHLOD(const FObjectInitializer& ObjectInitializer)
 	: Super(ObjectInitializer)
 	, bRequireWarmup(false)
@@ -72,6 +57,10 @@ AWorldPartitionHLOD::AWorldPartitionHLOD(const FObjectInitializer& ObjectInitial
 	HLODHash = 0;
 	HLODBounds = FBox(EForceInit::ForceInit);
 #endif
+
+#if WITH_EDITOR
+	FWorldDelegates::OnWorldCleanup.AddUObject(this, &AWorldPartitionHLOD::OnWorldCleanup);
+#endif
 }
 
 const FGuid& AWorldPartitionHLOD::GetSourceCellGuid() const
@@ -121,21 +110,33 @@ void AWorldPartitionHLOD::Serialize(FArchive& Ar)
 	Super::Serialize(Ar);
 
 #if WITH_EDITOR
-	if(Ar.IsLoading() && Ar.CustomVer(FUE5MainStreamObjectVersion::GUID) < FUE5MainStreamObjectVersion::WorldPartitionStreamingCellsNamingShortened)
+	if (Ar.IsLoading())
 	{
-		SourceCell_DEPRECATED = SourceCell_DEPRECATED.ToString().Replace(TEXT("WPRT_"), TEXT(""), ESearchCase::CaseSensitive).Replace(TEXT("Cell_"), TEXT(""), ESearchCase::CaseSensitive);
-	}
+		if (Ar.CustomVer(FUE5MainStreamObjectVersion::GUID) < FUE5MainStreamObjectVersion::WorldPartitionStreamingCellsNamingShortened)
+		{
+			SourceCell_DEPRECATED = SourceCell_DEPRECATED.ToString().Replace(TEXT("WPRT_"), TEXT(""), ESearchCase::CaseSensitive).Replace(TEXT("Cell_"), TEXT(""), ESearchCase::CaseSensitive);
+		}
 
-	if(Ar.IsLoading() && Ar.CustomVer(FUE5PrivateFrostyStreamObjectVersion::GUID) < FUE5PrivateFrostyStreamObjectVersion::ConvertWorldPartitionHLODsCellsToName)
-	{
-		FString CellName;
-		FString CellContext;
-		const FString CellPath = FPackageName::GetShortName(SourceCell_DEPRECATED.ToSoftObjectPath().GetSubPathString());
-		if (!CellPath.Split(TEXT("."), &CellContext, &CellName, ESearchCase::CaseSensitive, ESearchDir::FromEnd))
+		if (Ar.CustomVer(FUE5PrivateFrostyStreamObjectVersion::GUID) < FUE5PrivateFrostyStreamObjectVersion::ConvertWorldPartitionHLODsCellsToName)
 		{
-			CellName = CellPath;
+			FString CellName;
+			FString CellContext;
+			const FString CellPath = FPackageName::GetShortName(SourceCell_DEPRECATED.ToSoftObjectPath().GetSubPathString());
+			if (!CellPath.Split(TEXT("."), &CellContext, &CellName, ESearchCase::CaseSensitive, ESearchDir::FromEnd))
+			{
+				CellName = CellPath;
+			}
+			SourceCellName_DEPRECATED = *CellName;
+		}
+
+		if (Ar.CustomVer(FFortniteMainBranchObjectVersion::GUID) < FFortniteMainBranchObjectVersion::WorldPartitionHLODSourceActorsRefactor)
+		{
+			check(!SourceActors)
+			UWorldPartitionHLODSourceActorsFromCell* SourceActorsFromCell = NewObject<UWorldPartitionHLODSourceActorsFromCell>(this);
+			SourceActorsFromCell->SetActors(HLODSubActors_DEPRECATED);
+			SourceActorsFromCell->SetHLODLayer(SubActorsHLODLayer_DEPRECATED);
+			SourceActors = SourceActorsFromCell;
 		}
-		SourceCellName_DEPRECATED = *CellName;
 	}
 #endif
 }
@@ -160,11 +161,11 @@ void AWorldPartitionHLOD::PostLoad()
 #if WITH_EDITOR
 	if (GetLinkerCustomVersion(FFortniteMainBranchObjectVersion::GUID) < FFortniteMainBranchObjectVersion::WorldPartitionStreamingCellsNamingShortened)
 	{
-		if (!HLODSubActors.IsEmpty())
+		if (!HLODSubActors_DEPRECATED.IsEmpty())
 		{
 			// As we may be dealing with an unsaved world created from a template map, get
 			// the source package name of this HLOD actor and figure out the world name from there
-			FName ExternalActorsPath = HLODSubActors[0].ContainerPackage;
+			FName ExternalActorsPath = HLODSubActors_DEPRECATED[0].ContainerPackage;
 			FString WorldName = FPackageName::GetShortName(ExternalActorsPath);
 
 			// Strip "WorldName_" from the cell name
@@ -262,6 +263,31 @@ void AWorldPartitionHLOD::RerunConstructionScripts()
 {
 }
 
+void AWorldPartitionHLOD::OnWorldCleanup(UWorld* InWorld, bool bSessionEnded, bool bCleanupResources)
+{
+	// Close all asset editors associated with this HLOD actor
+	const UWorld* World = GetWorld();
+	if (World == InWorld)
+	{
+		if (!World->IsGameWorld())
+		{
+			UPackage* HLODPackage = GetPackage();
+
+			// Find all assets being edited
+			UAssetEditorSubsystem* AssetEditorSubsystem = GEditor->GetEditorSubsystem<UAssetEditorSubsystem>();
+			TArray<UObject*> AllAssets = AssetEditorSubsystem->GetAllEditedAssets();
+
+			for (UObject* Asset : AllAssets)
+			{
+				if (Asset->GetPackage() == HLODPackage)
+				{
+					AssetEditorSubsystem->CloseAllEditorsForAsset(Asset);
+				}
+			}
+		}
+	}
+}
+
 TUniquePtr<FWorldPartitionActorDesc> AWorldPartitionHLOD::CreateClassActorDesc() const
 {
 	return TUniquePtr<FWorldPartitionActorDesc>(new FHLODActorDesc());
@@ -321,20 +347,19 @@ void AWorldPartitionHLOD::SetHLODComponents(const TArray<UActorComponent*>& InHL
 	});
 }
 
-void AWorldPartitionHLOD::SetSubActors(const TArray<FHLODSubActor>& InHLODSubActors)
+void AWorldPartitionHLOD::SetSourceActors(UWorldPartitionHLODSourceActors* InHLODSourceActors)
 {
-	HLODSubActors = InHLODSubActors;
+	SourceActors = InHLODSourceActors;
 }
 
-const TArray<FHLODSubActor>& AWorldPartitionHLOD::GetSubActors() const
+const UWorldPartitionHLODSourceActors* AWorldPartitionHLOD::GetSourceActors() const
 {
-	return HLODSubActors;
+	return SourceActors;
 }
 
-void AWorldPartitionHLOD::SetSubActorsHLODLayer(const UHLODLayer* InSubActorsHLODLayer)
+UWorldPartitionHLODSourceActors* AWorldPartitionHLOD::GetSourceActors()
 {
-	SubActorsHLODLayer = InSubActorsHLODLayer;
-	bRequireWarmup = SubActorsHLODLayer->DoesRequireWarmup();
+	return SourceActors;
 }
 
 void AWorldPartitionHLOD::SetSourceCellGuid(const FGuid& InSourceCellGuid)
diff --git a/Engine/Source/Runtime/Engine/Private/WorldPartition/HLOD/HLODActorDesc.cpp b/Engine/Source/Runtime/Engine/Private/WorldPartition/HLOD/HLODActorDesc.cpp
--- a/Engine/Source/Runtime/Engine/Private/WorldPartition/HLOD/HLODActorDesc.cpp
+++ b/Engine/Source/Runtime/Engine/Private/WorldPartition/HLOD/HLODActorDesc.cpp
@@ -5,32 +5,79 @@
 #include "WorldPartition/HLOD/HLODStats.h"
 
 #if WITH_EDITOR
+
+#include "Engine/World.h"
 #include "HAL/FileManager.h"
 #include "Misc/PackageName.h"
 #include "UObject/UE5MainStreamObjectVersion.h"
 #include "UObject/UE5ReleaseStreamObjectVersion.h"
 #include "UObject/FortniteMainBranchObjectVersion.h"
+#include "WorldPartition/WorldPartition.h"
 #include "WorldPartition/HLOD/HLODActor.h"
 #include "WorldPartition/HLOD/HLODLayer.h"
+#include "WorldPartition/HLOD/HLODSourceActorsFromCell.h"
+
 
 void FHLODActorDesc::Init(const AActor* InActor)
 {
 	FWorldPartitionActorDesc::Init(InActor);
 
+	const UWorldPartition* WorldPartition = InActor->GetWorld() ? InActor->GetWorld()->GetWorldPartition() : nullptr;
 	const AWorldPartitionHLOD* HLODActor = CastChecked<AWorldPartitionHLOD>(InActor);
+	const UWorldPartitionHLODSourceActors* HLODSourceActors = HLODActor->GetSourceActors();
+	const UWorldPartitionHLODSourceActorsFromCell* HLODCellSourceActors = Cast<UWorldPartitionHLODSourceActorsFromCell>(HLODSourceActors);
 
-	HLODSubActors.Reserve(HLODActor->GetSubActors().Num());
-	Algo::Transform(HLODActor->GetSubActors(), HLODSubActors, [](const FHLODSubActor& SubActor) { return FHLODSubActorDesc(SubActor.ActorGuid, SubActor.ContainerID); });
+	if (HLODCellSourceActors && WorldPartition)
+	{
+		for (const FHLODSubActor& SubActor : HLODCellSourceActors->GetActors())
+		{
+			if (SubActor.ContainerID.IsMainContainer())
+			{
+				const FWorldPartitionActorDesc* SubActorDesc = WorldPartition->GetActorDesc(SubActor.ActorGuid);
+				if (SubActorDesc && SubActorDesc->GetActorNativeClass()->IsChildOf<AWorldPartitionHLOD>())
+				{
+					ChildHLODActors.Add(SubActor.ActorGuid);
+				}
+			}
+		}
+	}
 
-	SourceHLODLayerName = NAME_None;
-	if (const UHLODLayer* SourceHLODLayer = HLODActor->GetSubActorsHLODLayer())
+	if (HLODSourceActors)
 	{
-		SourceHLODLayerName = SourceHLODLayer->GetFName();
+		SourceHLODLayerName = NAME_None;
+		if (const UHLODLayer* SourceHLODLayer = HLODSourceActors->GetHLODLayer())
+		{
+			SourceHLODLayerName = SourceHLODLayer->GetFName();
+		}
 	}
 	
 	HLODStats = HLODActor->GetStats();
 }
 
+struct FHLODSubActorDescDeprecated
+{
+	friend FArchive& operator<<(FArchive& Ar, FHLODSubActorDescDeprecated& SubActor);
+	FGuid ActorGuid;
+	FActorContainerID ContainerID;
+};
+
+FArchive& operator<<(FArchive& Ar, FHLODSubActorDescDeprecated& SubActor)
+{
+	Ar << SubActor.ActorGuid;
+
+	if (Ar.CustomVer(FFortniteMainBranchObjectVersion::GUID) < FFortniteMainBranchObjectVersion::WorldPartitionFActorContainerIDu64ToGuid)
+	{
+		uint64 ContainerID;
+		Ar << ContainerID;
+	}
+	else
+	{
+		Ar << SubActor.ContainerID;
+	}
+
+	return Ar;
+}
+
 void FHLODActorDesc::Serialize(FArchive& Ar)
 {
 	Ar.UsingCustomVersion(FUE5MainStreamObjectVersion::GUID);
@@ -46,10 +93,15 @@ void FHLODActorDesc::Serialize(FArchive& Ar)
 			TArray<FGuid> SubActors;
 			Ar << SubActors;
 		}
-		else
+		else if(Ar.CustomVer(FFortniteMainBranchObjectVersion::GUID) < FFortniteMainBranchObjectVersion::WorldPartitionHLODSourceActorsRefactor)
 		{
+			TArray<FHLODSubActorDescDeprecated> HLODSubActors;
 			Ar << HLODSubActors;
 		}
+		else
+		{
+			Ar << ChildHLODActors;
+		}
 	
 		if (Ar.CustomVer(FUE5MainStreamObjectVersion::GUID) < FUE5MainStreamObjectVersion::WorldPartitionHLODActorDescSerializeHLODLayer)
 		{
@@ -97,7 +149,7 @@ bool FHLODActorDesc::Equals(const FWorldPartitionActorDesc* Other) const
 		const FHLODActorDesc& HLODActorDesc = *(FHLODActorDesc*)Other;
 		return SourceHLODLayerName == HLODActorDesc.SourceHLODLayerName &&
 			   HLODStats.OrderIndependentCompareEqual(HLODActorDesc.GetStats()) &&
-			   CompareUnsortedArrays(HLODSubActors, HLODActorDesc.HLODSubActors);
+			   CompareUnsortedArrays(ChildHLODActors, HLODActorDesc.ChildHLODActors);
 	}
 	return false;
 }
diff --git a/Engine/Source/Runtime/Engine/Private/WorldPartition/HLOD/HLODProviderInterface.cpp b/Engine/Source/Runtime/Engine/Private/WorldPartition/HLOD/HLODProviderInterface.cpp
new file mode 100644
--- /dev/null
+++ b/Engine/Source/Runtime/Engine/Private/WorldPartition/HLOD/HLODProviderInterface.cpp
@@ -0,0 +1,11 @@
+// Copyright Epic Games, Inc. All Rights Reserved.
+
+#include "WorldPartition/HLOD/HLODProviderInterface.h"
+
+#include UE_INLINE_GENERATED_CPP_BY_NAME(HLODProviderInterface)
+
+
+UWorldPartitionHLODProvider::UWorldPartitionHLODProvider(const FObjectInitializer& ObjectInitializer)
+	: Super(ObjectInitializer)
+{
+}
diff --git a/Engine/Source/Runtime/Engine/Private/WorldPartition/HLOD/HLODSourceActors.cpp b/Engine/Source/Runtime/Engine/Private/WorldPartition/HLOD/HLODSourceActors.cpp
new file mode 100644
--- /dev/null
+++ b/Engine/Source/Runtime/Engine/Private/WorldPartition/HLOD/HLODSourceActors.cpp
@@ -0,0 +1,62 @@
+// Copyright Epic Games, Inc. All Rights Reserved.
+
+#include "WorldPartition/HLOD/HLODSourceActors.h"
+
+#if WITH_EDITOR
+#include "WorldPartition/HLOD/HLODActor.h"
+#include "WorldPartition/HLOD/HLODLayer.h"
+#endif
+
+
+#if WITH_EDITOR
+
+static uint32 GetHLODHash(const UHLODLayer* InHLODLayer)
+{
+	UHLODLayer& HLODLayer = *const_cast<UHLODLayer*>(InHLODLayer);
+
+	uint32 HLODHash;
+
+	HLODHash = GetTypeHash(HLODLayer.GetLayerType());
+	UE_LOG(LogHLODHash, VeryVerbose, TEXT(" - LayerType = %d"), HLODHash);
+
+	HLODHash = HashCombine(HLODLayer.GetHLODBuilderSettings()->GetCRC(), HLODHash);
+	UE_LOG(LogHLODHash, VeryVerbose, TEXT(" - HLODBuilderSettings = %d"), HLODHash);
+
+	HLODHash = HashCombine(HLODLayer.GetCellSize(), HLODHash);
+	UE_LOG(LogHLODHash, VeryVerbose, TEXT(" - CellSize = %d"), HLODHash);
+
+	return HLODHash;
+}
+
+#endif // #if WITH_EDITOR
+
+UWorldPartitionHLODSourceActors::UWorldPartitionHLODSourceActors(const FObjectInitializer& ObjectInitializer)
+	: Super(ObjectInitializer)
+{
+}
+
+#if WITH_EDITOR
+
+uint32 UWorldPartitionHLODSourceActors::GetHLODHash() const
+{
+	uint32 HLODHash = 0;
+	if (HLODLayer)
+	{
+		// HLOD Layer
+		HLODHash = ::GetHLODHash(HLODLayer);
+		UE_LOG(LogHLODHash, VeryVerbose, TEXT(" - HLOD Layer (%s) = %x"), *HLODLayer->GetName(), HLODHash);
+	}
+	return HLODHash;
+}
+
+void UWorldPartitionHLODSourceActors::SetHLODLayer(const UHLODLayer* InHLODLayer)
+{
+	HLODLayer = InHLODLayer;
+}
+
+const UHLODLayer* UWorldPartitionHLODSourceActors::GetHLODLayer() const
+{
+	return HLODLayer;
+}
+
+#endif // #if WITH_EDITOR
\ No newline at end of file
diff --git a/Engine/Source/Runtime/Engine/Private/WorldPartition/HLOD/HLODSourceActorsFromCell.cpp b/Engine/Source/Runtime/Engine/Private/WorldPartition/HLOD/HLODSourceActorsFromCell.cpp
new file mode 100644
--- /dev/null
+++ b/Engine/Source/Runtime/Engine/Private/WorldPartition/HLOD/HLODSourceActorsFromCell.cpp
@@ -0,0 +1,95 @@
+// Copyright Epic Games, Inc. All Rights Reserved.
+
+#include "WorldPartition/HLOD/HLODSourceActorsFromCell.h"
+
+#if WITH_EDITOR
+#include "WorldPartition/HLOD/HLODActor.h"
+#include "WorldPartition/WorldPartitionLevelStreamingDynamic.h"
+#endif
+
+
+#if WITH_EDITOR
+
+static uint32 GetHLODHash(const TArray<FHLODSubActor>& InSubActors)
+{
+	uint32 HLODHash = 0;
+
+	for (const FHLODSubActor& HLODSubActor : InSubActors)
+	{
+		HLODHash = HashCombine(GetTypeHash(HLODSubActor), HLODHash);
+	}
+
+	return HLODHash;
+}
+
+#endif // #if WITH_EDITOR
+
+UWorldPartitionHLODSourceActorsFromCell::UWorldPartitionHLODSourceActorsFromCell(const FObjectInitializer& ObjectInitializer)
+	: Super(ObjectInitializer)
+{
+}
+
+#if WITH_EDITOR
+
+ULevelStreaming* UWorldPartitionHLODSourceActorsFromCell::LoadSourceActors(bool& bOutDirty) const
+{
+	UPackage::WaitForAsyncFileWrites();
+
+	bOutDirty = false;
+	AWorldPartitionHLOD* HLODActor = CastChecked<AWorldPartitionHLOD>(GetOuter());
+	UWorld* World = HLODActor->GetWorld();
+	UWorldPartition* WorldPartition = World->GetWorldPartition();
+	check(WorldPartition);
+
+	const FName LevelStreamingName = FName(*FString::Printf(TEXT("HLODLevelStreaming_%s"), *HLODActor->GetName()));
+	TArray<FWorldPartitionRuntimeCellObjectMapping> Mappings;
+	Mappings.Reserve(Actors.Num());
+	Algo::Transform(Actors, Mappings, [World, HLODActor](const FHLODSubActor& SubActor)
+	{
+		return FWorldPartitionRuntimeCellObjectMapping(
+			SubActor.ActorPackage,
+			SubActor.ActorPath,
+			SubActor.ContainerID,
+			SubActor.ContainerTransform,
+			SubActor.ContainerPackage,
+			World->GetPackage()->GetFName(),
+			HLODActor->GetContentBundleGuid(),
+			SubActor.ActorGuid,
+			false
+		);
+	});
+
+	UWorldPartitionLevelStreamingDynamic* LevelStreaming = UWorldPartitionLevelStreamingDynamic::LoadInEditor(World, LevelStreamingName, Mappings);
+	check(LevelStreaming);
+
+	if (!LevelStreaming->GetLoadSucceeded())
+	{
+		bOutDirty = true;
+	}
+
+	return LevelStreaming;
+}
+
+uint32 UWorldPartitionHLODSourceActorsFromCell::GetHLODHash() const
+{
+	uint32 HLODHash = Super::GetHLODHash();
+
+	// Sub Actors
+	uint32 SubActorsHash = ::GetHLODHash(Actors);
+	UE_LOG(LogHLODHash, VeryVerbose, TEXT(" - Sub Actors (%d actors) = %x"), Actors.Num(), SubActorsHash);
+	HLODHash = HashCombine(HLODHash, SubActorsHash);
+
+	return HLODHash;
+}
+
+void UWorldPartitionHLODSourceActorsFromCell::SetActors(const TArray<FHLODSubActor>& InSourceActors)
+{
+	Actors = InSourceActors;
+}
+
+const TArray<FHLODSubActor>& UWorldPartitionHLODSourceActorsFromCell::GetActors() const
+{
+	return Actors;
+}
+
+#endif // #if WITH_EDITOR
diff --git a/Engine/Source/Runtime/Engine/Private/WorldPartition/HLOD/HLODSourceActorsFromLevel.cpp b/Engine/Source/Runtime/Engine/Private/WorldPartition/HLOD/HLODSourceActorsFromLevel.cpp
new file mode 100644
--- /dev/null
+++ b/Engine/Source/Runtime/Engine/Private/WorldPartition/HLOD/HLODSourceActorsFromLevel.cpp
@@ -0,0 +1,57 @@
+// Copyright Epic Games, Inc. All Rights Reserved.
+
+#include "WorldPartition/HLOD/HLODSourceActorsFromLevel.h"
+
+#if WITH_EDITOR
+#include "WorldPartition/HLOD/HLODActor.h"
+#include "Engine/LevelStreamingDynamic.h"
+#include "Engine/World.h"
+#endif
+
+UWorldPartitionHLODSourceActorsFromLevel::UWorldPartitionHLODSourceActorsFromLevel(const FObjectInitializer& ObjectInitializer)
+	: Super(ObjectInitializer)
+{
+}
+
+#if WITH_EDITOR
+
+ULevelStreaming* UWorldPartitionHLODSourceActorsFromLevel::LoadSourceActors(bool& bOutDirty) const
+{
+	AWorldPartitionHLOD* HLODActor = CastChecked<AWorldPartitionHLOD>(GetOuter());
+	UWorld* World = HLODActor->GetWorld();
+
+	bool bSuccess = false;
+	FName LevelStreamingName = SourceLevel.GetUniqueID().GetLongPackageFName();
+	ULevelStreamingDynamic* LevelStreaming = ULevelStreamingDynamic::LoadLevelInstanceBySoftObjectPtr(World, SourceLevel, FTransform::Identity, bSuccess, TEXT(""), nullptr, true);
+
+	if (bSuccess && LevelStreaming)
+	{
+		World->FlushLevelStreaming();
+		return LevelStreaming;
+	}
+
+	return nullptr;
+}
+
+uint32 UWorldPartitionHLODSourceActorsFromLevel::GetHLODHash() const
+{
+	uint32 HLODHash = Super::GetHLODHash();
+
+	// SourceLevel
+	HLODHash = HashCombine(HLODHash, GetTypeHash(SourceLevel));
+	
+	return HLODHash;
+}
+
+
+void UWorldPartitionHLODSourceActorsFromLevel::SetSourceLevel(const UWorld* InSourceLevel)
+{
+	SourceLevel = InSourceLevel;
+}
+
+const TSoftObjectPtr<UWorld>& UWorldPartitionHLODSourceActorsFromLevel::GetSourceLevel() const
+{
+	return SourceLevel;
+}
+
+#endif // #if WITH_EDITOR
\ No newline at end of file
diff --git a/Engine/Source/Runtime/Engine/Public/WorldPartition/HLOD/HLODActor.h b/Engine/Source/Runtime/Engine/Public/WorldPartition/HLOD/HLODActor.h
--- a/Engine/Source/Runtime/Engine/Public/WorldPartition/HLOD/HLODActor.h
+++ b/Engine/Source/Runtime/Engine/Public/WorldPartition/HLOD/HLODActor.h
@@ -1,4 +1,5 @@
 // Copyright Epic Games, Inc. All Rights Reserved.
+
 #pragma once
 
 #include "GameFramework/Actor.h"
@@ -12,6 +13,9 @@
 
 class UHLODLayer;
 class UHLODSubsystem;
+class UWorldPartitionHLODSourceActors;
+
+ENGINE_API DECLARE_LOG_CATEGORY_EXTERN(LogHLODHash, Log, All);
 
 UCLASS(NotPlaceable)
 class ENGINE_API AWorldPartitionHLOD : public AActor
@@ -33,11 +37,9 @@ public:
 #if WITH_EDITOR
 	void SetHLODComponents(const TArray<UActorComponent*>& InHLODComponents);
 
-	void SetSubActors(const TArray<FHLODSubActor>& InSubActorMappings);
-	const TArray<FHLODSubActor>& GetSubActors() const;
-
-	void SetSubActorsHLODLayer(const UHLODLayer* InSubActorsHLODLayer);
-	const UHLODLayer* GetSubActorsHLODLayer() const { return SubActorsHLODLayer; }
+	void SetSourceActors(UWorldPartitionHLODSourceActors* InSourceActors);
+	UWorldPartitionHLODSourceActors* GetSourceActors();
+	const UWorldPartitionHLODSourceActors* GetSourceActors() const;
 
 	void SetRequireWarmup(bool InRequireWarmup) { bRequireWarmup = InRequireWarmup; }
 
@@ -89,12 +91,13 @@ protected:
 	//~ End AActor Interface.
 
 private:
-#if WITH_EDITORONLY_DATA
-	UPROPERTY()
-	TArray<FHLODSubActor> HLODSubActors;
+#if WITH_EDITOR
+	void OnWorldCleanup(UWorld* InWorld, bool bSessionEnded, bool bCleanupResources);
+#endif
 
+#if WITH_EDITORONLY_DATA
 	UPROPERTY()
-	TObjectPtr<const UHLODLayer> SubActorsHLODLayer;
+	TObjectPtr<UWorldPartitionHLODSourceActors> SourceActors;
 
 	UPROPERTY()
 	FBox HLODBounds;
@@ -115,6 +118,10 @@ private:
 	UPROPERTY()
 	bool bRequireWarmup;
 
+	UPROPERTY()
+	FGuid SourceCellGuid;
+
+#if WITH_EDITORONLY_DATA
 	UPROPERTY()
 	TSoftObjectPtr<UWorldPartitionRuntimeCell> SourceCell_DEPRECATED;
 
@@ -122,7 +129,11 @@ private:
 	FName SourceCellName_DEPRECATED;
 
 	UPROPERTY()
-	FGuid SourceCellGuid;
+	TArray<FHLODSubActor> HLODSubActors_DEPRECATED;
+
+	UPROPERTY()
+	TObjectPtr<const UHLODLayer> SubActorsHLODLayer_DEPRECATED;
+#endif	
 };
 
 DEFINE_ACTORDESC_TYPE(AWorldPartitionHLOD, FHLODActorDesc);
\ No newline at end of file
diff --git a/Engine/Source/Runtime/Engine/Public/WorldPartition/HLOD/HLODActorDesc.h b/Engine/Source/Runtime/Engine/Public/WorldPartition/HLOD/HLODActorDesc.h
--- a/Engine/Source/Runtime/Engine/Public/WorldPartition/HLOD/HLODActorDesc.h
+++ b/Engine/Source/Runtime/Engine/Public/WorldPartition/HLOD/HLODActorDesc.h
@@ -25,7 +25,7 @@ class ENGINE_API FHLODActorDesc : public FWorldPartitionActorDesc
 public:
 	typedef TMap<FName, int64>	FStats;
 
-	inline const TArray<FHLODSubActorDesc>& GetSubActors() const { return HLODSubActors; }
+	inline const TArray<FGuid>& GetChildHLODActors() const { return ChildHLODActors; }
 	inline const FName GetSourceHLODLayerName() const { return SourceHLODLayerName; }
 	inline const FStats& GetStats() const { return HLODStats; }
 	inline int64 GetStat(FName InStatName) const { return HLODStats.FindRef(InStatName); }
@@ -43,7 +43,7 @@ protected:
 	virtual bool ShouldValidateRuntimeGrid() const override { return false; }
 	//~ End FWorldPartitionActorDesc Interface.
 
-	TArray<FHLODSubActorDesc> HLODSubActors;
+	TArray<FGuid> ChildHLODActors;
 	FName SourceHLODLayerName;
 	FStats HLODStats;
 };
diff --git a/Engine/Source/Runtime/Engine/Public/WorldPartition/HLOD/HLODProviderInterface.h b/Engine/Source/Runtime/Engine/Public/WorldPartition/HLOD/HLODProviderInterface.h
new file mode 100644
--- /dev/null
+++ b/Engine/Source/Runtime/Engine/Public/WorldPartition/HLOD/HLODProviderInterface.h
@@ -0,0 +1,26 @@
+// Copyright Epic Games, Inc. All Rights Reserved.
+
+#pragma once
+
+#include "CoreMinimal.h"
+#include "UObject/Interface.h"
+#include "HLODProviderInterface.generated.h"
+
+
+class AWorldPartitionHLOD;
+
+
+UINTERFACE()
+class ENGINE_API UWorldPartitionHLODProvider : public UInterface
+{
+	GENERATED_UINTERFACE_BODY()
+};
+
+
+class ENGINE_API IWorldPartitionHLODProvider
+{
+	GENERATED_IINTERFACE_BODY()
+
+public:
+	virtual AWorldPartitionHLOD* CreateHLODActor() = 0;
+};
diff --git a/Engine/Source/Runtime/Engine/Public/WorldPartition/HLOD/HLODSourceActors.h b/Engine/Source/Runtime/Engine/Public/WorldPartition/HLOD/HLODSourceActors.h
new file mode 100644
--- /dev/null
+++ b/Engine/Source/Runtime/Engine/Public/WorldPartition/HLOD/HLODSourceActors.h
@@ -0,0 +1,31 @@
+// Copyright Epic Games, Inc. All Rights Reserved.
+
+#pragma once
+
+#include "HLODSourceActors.generated.h"
+
+
+class ULevelStreaming;
+class UHLODLayer;
+
+
+UCLASS(Abstract)
+class ENGINE_API UWorldPartitionHLODSourceActors : public UObject
+{
+	GENERATED_UCLASS_BODY()
+
+public:
+#if WITH_EDITOR
+	virtual ULevelStreaming* LoadSourceActors(bool& bOutDirty) const PURE_VIRTUAL(UWorldPartitionHLODSourceActors::LoadSourceActors, return nullptr; );
+	virtual uint32 GetHLODHash() const;
+
+	void SetHLODLayer(const UHLODLayer* HLODLayer);
+	const UHLODLayer* GetHLODLayer() const;
+#endif
+
+private:
+#if WITH_EDITORONLY_DATA
+	UPROPERTY()
+	TObjectPtr<const UHLODLayer> HLODLayer;
+#endif
+};
diff --git a/Engine/Source/Runtime/Engine/Public/WorldPartition/HLOD/HLODSourceActorsFromCell.h b/Engine/Source/Runtime/Engine/Public/WorldPartition/HLOD/HLODSourceActorsFromCell.h
new file mode 100644
--- /dev/null
+++ b/Engine/Source/Runtime/Engine/Public/WorldPartition/HLOD/HLODSourceActorsFromCell.h
@@ -0,0 +1,29 @@
+// Copyright Epic Games, Inc. All Rights Reserved.
+
+#pragma once
+
+#include "WorldPartition/HLOD/HLODSourceActors.h"
+#include "WorldPartition/HLOD/HLODSubActor.h"
+#include "HLODSourceActorsFromCell.generated.h"
+
+
+UCLASS()
+class ENGINE_API UWorldPartitionHLODSourceActorsFromCell : public UWorldPartitionHLODSourceActors
+{
+	GENERATED_UCLASS_BODY()
+
+public:
+#if WITH_EDITOR
+	virtual ULevelStreaming* LoadSourceActors(bool& bOutDirty) const override;
+	virtual uint32 GetHLODHash() const override;
+
+	void SetActors(const TArray<FHLODSubActor>& InSourceActors);
+	const TArray<FHLODSubActor>& GetActors() const;
+#endif
+
+private:
+#if WITH_EDITORONLY_DATA
+	UPROPERTY()
+	TArray<FHLODSubActor> Actors;
+#endif
+};
diff --git a/Engine/Source/Runtime/Engine/Public/WorldPartition/HLOD/HLODSourceActorsFromLevel.h b/Engine/Source/Runtime/Engine/Public/WorldPartition/HLOD/HLODSourceActorsFromLevel.h
new file mode 100644
--- /dev/null
+++ b/Engine/Source/Runtime/Engine/Public/WorldPartition/HLOD/HLODSourceActorsFromLevel.h
@@ -0,0 +1,28 @@
+// Copyright Epic Games, Inc. All Rights Reserved.
+
+#pragma once
+
+#include "WorldPartition/HLOD/HLODSourceActors.h"
+#include "HLODSourceActorsFromLevel.generated.h"
+
+
+UCLASS()
+class ENGINE_API UWorldPartitionHLODSourceActorsFromLevel : public UWorldPartitionHLODSourceActors
+{
+	GENERATED_UCLASS_BODY()
+
+public:
+#if WITH_EDITOR
+	virtual ULevelStreaming* LoadSourceActors(bool& bOutDirty) const override;
+	virtual uint32 GetHLODHash() const override;
+
+	void SetSourceLevel(const UWorld* InSourceLevel);
+	const TSoftObjectPtr<UWorld>& GetSourceLevel() const;
+#endif
+
+private:
+#if WITH_EDITORONLY_DATA
+	UPROPERTY()
+	TSoftObjectPtr<UWorld> SourceLevel;
+#endif
+};
diff --git a/Engine/Source/Runtime/Engine/Public/WorldPartition/HLOD/HLODSubActor.h b/Engine/Source/Runtime/Engine/Public/WorldPartition/HLOD/HLODSubActor.h
--- a/Engine/Source/Runtime/Engine/Public/WorldPartition/HLOD/HLODSubActor.h
+++ b/Engine/Source/Runtime/Engine/Public/WorldPartition/HLOD/HLODSubActor.h
@@ -5,51 +5,6 @@
 #include "WorldPartition/WorldPartitionActorContainerID.h"
 #include "HLODSubActor.generated.h"
 
-USTRUCT()
-struct FHLODSubActorDesc
-{
-	GENERATED_USTRUCT_BODY()
-
-	FHLODSubActorDesc()
-	{}
-
-	FHLODSubActorDesc(FGuid InActorGuid, const FActorContainerID& InContainerID)
-#if WITH_EDITORONLY_DATA
-		: ActorGuid(InActorGuid)
-		, ContainerID(InContainerID)
-#endif
-	{}
-
-#if WITH_EDITORONLY_DATA
-	friend bool operator==(const FHLODSubActorDesc& Lhs, const FHLODSubActorDesc& Rhs)
-	{
-		return Lhs.ActorGuid == Rhs.ActorGuid && Lhs.ContainerID == Rhs.ContainerID;
-	}
-
-	friend bool operator!=(const FHLODSubActorDesc& Lhs, const FHLODSubActorDesc& Rhs)
-	{
-		return !(Lhs == Rhs);
-	}
-
-	friend bool operator<(const FHLODSubActorDesc& Lhs, const FHLODSubActorDesc& Rhs)
-	{
-		if (Lhs.ActorGuid == Rhs.ActorGuid)
-		{
-			return Lhs.ContainerID < Rhs.ContainerID;
-		}
-
-		return Lhs.ActorGuid < Rhs.ActorGuid;
-	}
-	
-	ENGINE_API friend FArchive& operator<<(FArchive& Ar, FHLODSubActorDesc& SubActor);
-
-	UPROPERTY()
-	FGuid ActorGuid;
-
-	UPROPERTY()
-	FActorContainerID ContainerID;
-#endif
-};
 
 USTRUCT()
 struct FHLODSubActor
diff --git a/Engine/Source/Runtime/Engine/Public/WorldPartition/HLOD/IWorldPartitionHLODUtilities.h b/Engine/Source/Runtime/Engine/Public/WorldPartition/HLOD/IWorldPartitionHLODUtilities.h
--- a/Engine/Source/Runtime/Engine/Public/WorldPartition/HLOD/IWorldPartitionHLODUtilities.h
+++ b/Engine/Source/Runtime/Engine/Public/WorldPartition/HLOD/IWorldPartitionHLODUtilities.h
@@ -27,8 +27,6 @@ struct FHLODCreationParams
 {
 	UWorldPartition* WorldPartition;
 
-	FName HLODLayerName;
-
 	FGuid CellGuid;
 	FString CellName;
 	FBox  CellBounds;
