World Partition - HLOD: During HLOD builds, wait for async file writes before loading source actors

diff --git a/Engine/Plugins/Editor/WorldPartitionHLODUtilities/Source/Private/WorldPartition/HLOD/Utilities/WorldPartitionHLODUtilities.cpp b/Engine/Plugins/Editor/WorldPartitionHLODUtilities/Source/Private/WorldPartition/HLOD/Utilities/WorldPartitionHLODUtilities.cpp
--- a/Engine/Plugins/Editor/WorldPartitionHLODUtilities/Source/Private/WorldPartition/HLOD/Utilities/WorldPartitionHLODUtilities.cpp
+++ b/Engine/Plugins/Editor/WorldPartitionHLODUtilities/Source/Private/WorldPartition/HLOD/Utilities/WorldPartitionHLODUtilities.cpp
@@ -574,6 +574,13 @@ static TArray<UActorComponent*> GatherHLODRelevantComponents(const TArray<AActor
 
 ULevelStreaming* LoadSourceActors(AWorldPartitionHLOD* InHLODActor, bool& bOutIsDirty)
 {
+	// If we're loading actors for an HLOD > 0
+	// Ensure that async writes for source actors (which are HLODs N-1) are completed.
+	if (InHLODActor->GetLODLevel() > 0)
+	{
+		UPackage::WaitForAsyncFileWrites();
+	}
+	
 	ULevelStreaming* LevelStreaming = InHLODActor->GetSourceActors()->LoadSourceActors(bOutIsDirty);
 	UE_CLOG(bOutIsDirty, LogHLODBuilder, Warning, TEXT("HLOD actor \"%s\" needs to be rebuilt as it didn't succeed in loading all actors."), *InHLODActor->GetActorLabel());
 
