VK - HLODs freeze when trying to build
* Prevent blocking call to UTexture::ForceUpdateTextureStreaming(). Some requests require the main thread to be ticked in order to run to completion.

diff --git a/Engine/Plugins/Editor/WorldPartitionHLODUtilities/Source/Private/WorldPartition/HLOD/Utilities/WorldPartitionHLODUtilities.cpp b/Engine/Plugins/Editor/WorldPartitionHLODUtilities/Source/Private/WorldPartition/HLOD/Utilities/WorldPartitionHLODUtilities.cpp
--- a/Engine/Plugins/Editor/WorldPartitionHLODUtilities/Source/Private/WorldPartition/HLOD/Utilities/WorldPartitionHLODUtilities.cpp
+++ b/Engine/Plugins/Editor/WorldPartitionHLODUtilities/Source/Private/WorldPartition/HLOD/Utilities/WorldPartitionHLODUtilities.cpp
@@ -20,6 +20,7 @@
 #include "WorldPartition/HLOD/HLODStats.h"
 #include "WorldPartition/HLOD/HLODSourceActorsFromCell.h"
 #include "WorldPartition/ContentBundle/ContentBundleActivationScope.h"
+#include "WorldPartition/WorldPartitionHelpers.h"
 
 #include "WorldPartition/HLOD/Builders/HLODBuilderInstancing.h"
 #include "WorldPartition/HLOD/Builders/HLODBuilderMeshMerge.h"
@@ -38,6 +39,7 @@
 #include "TextureCompiler.h"
 #include "UObject/MetaData.h"
 #include "UObject/GCObjectScopeGuard.h"
+#include "ContentStreaming.h"
 
 static uint32 ComputeHLODHash(AWorldPartitionHLOD* InHLODActor, const TArray<UActorComponent*>& InSourceComponents)
 {
@@ -590,6 +592,14 @@ ULevelStreaming* LoadSourceActors(AWorldPartitionHLOD* InHLODActor, bool& bOutIs
 
 		// Ensure all deferred construction scripts are executed
 		FAssetCompilingManager::Get().ProcessAsyncTasks();
+
+		// Ensure streaming requests are completed
+		// Some requests require the main thread to be ticked in order to run to completion.
+		const double STREAMING_WAIT_DT = 0.1;
+		while (IStreamingManager::Get().StreamAllResources(STREAMING_WAIT_DT) > 0)
+		{
+			FWorldPartitionHelpers::FakeEngineTick(InHLODActor->GetWorld());
+		}
 	}
 
 	return LevelStreaming;
diff --git a/Engine/Source/Developer/MaterialBaking/Private/MaterialBakingModule.cpp b/Engine/Source/Developer/MaterialBaking/Private/MaterialBakingModule.cpp
--- a/Engine/Source/Developer/MaterialBaking/Private/MaterialBakingModule.cpp
+++ b/Engine/Source/Developer/MaterialBaking/Private/MaterialBakingModule.cpp
@@ -690,8 +690,17 @@ private:
 			}
 		}
 
+
 		// Force all mip maps to load before baking the materials
-		UTexture::ForceUpdateTextureStreaming();
+		{
+			const double STREAMING_WAIT_DT = 0.1;
+			while (IStreamingManager::Get().StreamAllResources(STREAMING_WAIT_DT) > 0)
+			{
+				// Application tick.
+				FTaskGraphInterface::Get().ProcessThreadUntilIdle(ENamedThreads::GameThread);
+				FTSTicker::GetCoreTicker().Tick(FApp::GetDeltaTime());
+			}
+		}
 	}
 
 	void BakeMaterialProperty(const FMaterialDataEx& CurrentMaterialSettings, const FMaterialPropertyEx& Property, FMeshMaterialRenderItem* RenderItem, UTextureRenderTarget2D* RenderTarget, FExportMaterialProxy* ExportMaterialProxy, FBakeOutputEx& CurrentOutput)
diff --git a/Engine/Source/Runtime/Engine/Private/ContentStreaming.cpp b/Engine/Source/Runtime/Engine/Private/ContentStreaming.cpp
--- a/Engine/Source/Runtime/Engine/Private/ContentStreaming.cpp
+++ b/Engine/Source/Runtime/Engine/Private/ContentStreaming.cpp
@@ -727,6 +727,11 @@ void IStreamingManager::Tick( float DeltaTime, bool bProcessEverything/*=false*/
 	bPendingRemoveViews = true;
 }
 
+int32 IStreamingManager::StreamAllResources(float TimeLimit)
+{
+	return 0;
+}
+
 /*-----------------------------------------------------------------------------
 	IRenderAssetStreamingManager implementation.
 -----------------------------------------------------------------------------*/
diff --git a/Engine/Source/Runtime/Engine/Public/ContentStreaming.h b/Engine/Source/Runtime/Engine/Public/ContentStreaming.h
--- a/Engine/Source/Runtime/Engine/Public/ContentStreaming.h
+++ b/Engine/Source/Runtime/Engine/Public/ContentStreaming.h
@@ -200,11 +200,8 @@ struct IStreamingManager
 	 * @param TimeLimit					Maximum number of seconds to wait for streaming I/O. If zero, uses .ini setting
 	 * @return							Number of streaming requests still in flight, if the time limit was reached before they were finished.
 	 */
-	virtual int32 StreamAllResources(float TimeLimit = 0.0f)
-	{
-		return 0;
-	}
-
+	ENGINE_API virtual int32 StreamAllResources(float TimeLimit = 0.0f);
+	
 	/**
 	 * Blocks till all pending requests are fulfilled.
 	 *
