World Partition HLOD - Approximate actors: Fixed issue with conflicting objects when building HLODs
* Generate assets in a temp package, then rename

diff --git a/Engine/Plugins/Editor/WorldPartitionHLODUtilities/Source/Private/WorldPartition/HLOD/Builders/HLODBuilderMeshApproximate.cpp b/Engine/Plugins/Editor/WorldPartitionHLODUtilities/Source/Private/WorldPartition/HLOD/Builders/HLODBuilderMeshApproximate.cpp
--- a/Engine/Plugins/Editor/WorldPartitionHLODUtilities/Source/Private/WorldPartition/HLOD/Builders/HLODBuilderMeshApproximate.cpp
+++ b/Engine/Plugins/Editor/WorldPartitionHLODUtilities/Source/Private/WorldPartition/HLOD/Builders/HLODBuilderMeshApproximate.cpp
@@ -19,6 +19,7 @@
 #include "Materials/Material.h"
 #include "Engine/HLODProxy.h"
 #include "Serialization/ArchiveCrc32.h"
+#include "ObjectTools.h"
 
 #include UE_INLINE_GENERATED_CPP_BY_NAME(HLODBuilderMeshApproximate)
 
@@ -116,6 +117,12 @@ TArray<UActorComponent*> UHLODBuilderMeshApproximate::Build(const FHLODBuildCont
 		Options.TextureSizePolicy = IGeometryProcessing_ApproximateActors::ETextureSizePolicy::TexelDensity;
 	}
 
+	// Use temp packages - Needed to allow proper replacement of existing assets (performed below)
+	const FString NewAssetNamePrefix(TEXT("NEWASSET_"));
+	FString PackagePath = InHLODBuildContext.AssetsOuter->GetPackage()->GetName();
+	FString AssetName = InHLODBuildContext.AssetsBaseName;
+	Options.BasePackagePath = PackagePath / NewAssetNamePrefix + AssetName;
+
 	// run actor approximation computation
 	IGeometryProcessing_ApproximateActors::FResults Results;
 	ApproxActorsAPI->ApproximateActors(Input, Options, Results);
@@ -123,15 +130,36 @@ TArray<UActorComponent*> UHLODBuilderMeshApproximate::Build(const FHLODBuildCont
 	TArray<UActorComponent*> Components;
 	if (Results.ResultCode == IGeometryProcessing_ApproximateActors::EResultCode::Success)
 	{
-		auto FixupAsset = [InHLODBuildContext](UObject* Asset)
+		auto ProcessNewAsset = [&InHLODBuildContext, &NewAssetNamePrefix](UObject* NewAsset)
 		{
-			Asset->ClearFlags(RF_Public | RF_Standalone);
-			Asset->Rename(nullptr, InHLODBuildContext.AssetsOuter, REN_NonTransactional | REN_DontCreateRedirectors | REN_ForceNoResetLoaders);
+			// Move asset out of the temp package and into its final package
+			{
+				FString AssetName = NewAsset->GetName();
+				AssetName.RemoveFromStart(NewAssetNamePrefix);
+
+				UObject* AssetToReplace = StaticFindObjectFast(UObject::StaticClass(), InHLODBuildContext.AssetsOuter, *AssetName);
+				if (AssetToReplace)
+				{
+					// Move the previous asset to the transient package
+					AssetToReplace->Rename(nullptr, GetTransientPackage(), REN_DontCreateRedirectors | REN_NonTransactional | REN_ForceNoResetLoaders);
+				}
+
+				UPackage* TempPackage = NewAsset->GetPackage();
+
+				// Rename the asset to its final destination
+				NewAsset->Rename(*AssetName, InHLODBuildContext.AssetsOuter, REN_DontCreateRedirectors | REN_NonTransactional | REN_ForceNoResetLoaders);
+				NewAsset->ClearFlags(RF_Public | RF_Standalone);
+
+				// Clean up flags on the temp package. It is not useful anymore.
+				TempPackage->ClearDirtyFlag();
+				TempPackage->SetFlags(RF_Transient);
+				TempPackage->ClearFlags(RF_Public | RF_Standalone);
+			}
 		};
 	
-		Algo::ForEach(Results.NewMeshAssets, FixupAsset);
-		Algo::ForEach(Results.NewMaterials, FixupAsset);
-		Algo::ForEach(Results.NewTextures, FixupAsset);
+		Algo::ForEach(Results.NewMeshAssets, ProcessNewAsset);
+		Algo::ForEach(Results.NewMaterials, ProcessNewAsset);
+		Algo::ForEach(Results.NewTextures, ProcessNewAsset);
 
 		for (UStaticMesh* StaticMesh : Results.NewMeshAssets)
 		{
