World Partition - HLOD:  Avoid issue with LoadSourceActors() ticking the whole engine
* Doesn't work well when called inside a running editor, and is unnecesary

diff --git a/Engine/Plugins/Editor/WorldPartitionHLODUtilities/Source/Private/WorldPartition/HLOD/Utilities/WorldPartitionHLODUtilities.cpp b/Engine/Plugins/Editor/WorldPartitionHLODUtilities/Source/Private/WorldPartition/HLOD/Utilities/WorldPartitionHLODUtilities.cpp
--- a/Engine/Plugins/Editor/WorldPartitionHLODUtilities/Source/Private/WorldPartition/HLOD/Utilities/WorldPartitionHLODUtilities.cpp
+++ b/Engine/Plugins/Editor/WorldPartitionHLODUtilities/Source/Private/WorldPartition/HLOD/Utilities/WorldPartitionHLODUtilities.cpp
@@ -20,7 +20,6 @@
 #include "WorldPartition/HLOD/HLODStats.h"
 #include "WorldPartition/HLOD/HLODSourceActorsFromCell.h"
 #include "WorldPartition/ContentBundle/ContentBundleActivationScope.h"
-#include "WorldPartition/WorldPartitionHelpers.h"
 
 #include "WorldPartition/HLOD/Builders/HLODBuilderInstancing.h"
 #include "WorldPartition/HLOD/Builders/HLODBuilderMeshMerge.h"
@@ -598,7 +597,9 @@ ULevelStreaming* LoadSourceActors(AWorldPartitionHLOD* InHLODActor, bool& bOutIs
 		const double STREAMING_WAIT_DT = 0.1;
 		while (IStreamingManager::Get().StreamAllResources(STREAMING_WAIT_DT) > 0)
 		{
-			FWorldPartitionHelpers::FakeEngineTick(InHLODActor->GetWorld());
+			// Application tick.
+			FTaskGraphInterface::Get().ProcessThreadUntilIdle(ENamedThreads::GameThread);
+			FTSTicker::GetCoreTicker().Tick(FApp::GetDeltaTime());
 		}
 	}
 
