World Partition - HLOD: Make sure the Approximate Actors HLODs can use "legacy" flatten materials without issue
* Those materials use different parameters for textures & switches names in the material.

diff --git a/Engine/Plugins/Editor/WorldPartitionHLODUtilities/Source/Private/WorldPartition/HLOD/Builders/HLODBuilderMeshApproximate.cpp b/Engine/Plugins/Editor/WorldPartitionHLODUtilities/Source/Private/WorldPartition/HLOD/Builders/HLODBuilderMeshApproximate.cpp
--- a/Engine/Plugins/Editor/WorldPartitionHLODUtilities/Source/Private/WorldPartition/HLOD/Builders/HLODBuilderMeshApproximate.cpp
+++ b/Engine/Plugins/Editor/WorldPartitionHLODUtilities/Source/Private/WorldPartition/HLOD/Builders/HLODBuilderMeshApproximate.cpp
@@ -11,12 +11,12 @@
 #include "Engine/Texture2D.h"
 #include "Materials/Material.h"
 #include "Materials/MaterialInstance.h"
+#include "MaterialUtilities.h"
 
 #include "Modules/ModuleManager.h"
 #include "IGeometryProcessingInterfacesModule.h"
 #include "GeometryProcessingInterfaces/ApproximateActors.h"
 
-#include "Materials/Material.h"
 #include "Engine/HLODProxy.h"
 #include "Serialization/ArchiveCrc32.h"
 #include "ObjectTools.h"
@@ -101,12 +101,16 @@ TArray<UActorComponent*> UHLODBuilderMeshApproximate::Build(const FHLODBuildCont
 
 	// Material baking settings
 	Options.BakeMaterial = HLODMaterial;
-	Options.BaseColorTexParamName = FName("BaseColorTexture");
-	Options.NormalTexParamName = FName("NormalTexture");
-	Options.MetallicTexParamName = FName("MetallicTexture");
-	Options.RoughnessTexParamName = FName("RoughnessTexture");
-	Options.SpecularTexParamName = FName("SpecularTexture");
-	Options.EmissiveTexParamName = FName("EmissiveHDRTexture");
+	if (!FMaterialUtilities::IsValidFlattenMaterial(Options.BakeMaterial))
+	{
+		Options.BakeMaterial = GEngine->DefaultFlattenMaterial;
+	}
+	Options.BaseColorTexParamName = FName(FMaterialUtilities::GetFlattenMaterialTextureName(EFlattenMaterialProperties::Diffuse, Options.BakeMaterial));
+	Options.NormalTexParamName = FName(FMaterialUtilities::GetFlattenMaterialTextureName(EFlattenMaterialProperties::Normal, Options.BakeMaterial));
+	Options.MetallicTexParamName = FName(FMaterialUtilities::GetFlattenMaterialTextureName(EFlattenMaterialProperties::Metallic, Options.BakeMaterial));
+	Options.RoughnessTexParamName = FName(FMaterialUtilities::GetFlattenMaterialTextureName(EFlattenMaterialProperties::Roughness, Options.BakeMaterial));
+	Options.SpecularTexParamName = FName(FMaterialUtilities::GetFlattenMaterialTextureName(EFlattenMaterialProperties::Specular, Options.BakeMaterial));
+	Options.EmissiveTexParamName = FName(FMaterialUtilities::GetFlattenMaterialTextureName(EFlattenMaterialProperties::Emissive, Options.BakeMaterial));
 	Options.bUsePackedMRS = true;
 	Options.PackedMRSTexParamName = FName("PackedTexture");
 
@@ -189,9 +193,11 @@ TArray<UActorComponent*> UHLODBuilderMeshApproximate::Build(const FHLODBuildCont
 
 			// Set proper switches needed by our base flatten material
 			SetStaticSwitch("UseBaseColor", Options.bBakeBaseColor);
+			SetStaticSwitch("UseDiffuse", Options.bBakeBaseColor);
 			SetStaticSwitch("UseRoughness", Options.bBakeRoughness);
 			SetStaticSwitch("UseMetallic", Options.bBakeMetallic);
 			SetStaticSwitch("UseSpecular", Options.bBakeSpecular);
+			SetStaticSwitch("UseEmissive", Options.bBakeEmissive);
 			SetStaticSwitch("UseEmissiveHDR", Options.bBakeEmissive);
 			SetStaticSwitch("UseNormal", Options.bBakeNormalMap);
 			SetStaticSwitch("PackMetallic", Options.bUsePackedMRS);
diff --git a/Engine/Plugins/Experimental/MeshModelingToolsetExp/Source/GeometryProcessingAdapters/Private/ApproximateActorsImpl.cpp b/Engine/Plugins/Experimental/MeshModelingToolsetExp/Source/GeometryProcessingAdapters/Private/ApproximateActorsImpl.cpp
--- a/Engine/Plugins/Experimental/MeshModelingToolsetExp/Source/GeometryProcessingAdapters/Private/ApproximateActorsImpl.cpp
+++ b/Engine/Plugins/Experimental/MeshModelingToolsetExp/Source/GeometryProcessingAdapters/Private/ApproximateActorsImpl.cpp
@@ -111,6 +111,10 @@ static TUniquePtr<FSceneCapturePhotoSet> CapturePhotoSet(
 		SceneCapture->SetCaptureTypeEnabled(ERenderCaptureType::Specular, bSpecular);
 	}
 
+	// These capture types aren't yet supported by the Approximate Actors interface
+	SceneCapture->SetCaptureTypeEnabled(ERenderCaptureType::Opacity, false);
+	SceneCapture->SetCaptureTypeEnabled(ERenderCaptureType::SubsurfaceColor, false);
+
 	UWorld* World = Input.Actors.IsEmpty() ? (Input.Components.IsEmpty() ? nullptr : Input.Components[0]->GetWorld()) : Input.Actors[0]->GetWorld();
 
 	SceneCapture->SetCaptureSceneActorsAndComponents(World, Input.Actors, Input.Components);
@@ -1275,7 +1279,10 @@ void FApproximateActorsImpl::GenerateApproximationForActorSet(const FInput& Inpu
 
 	// force material update now that we have updated texture parameters
 	// (does this do that? Let calling code do it?)
-	NewMaterial->PostEditChange();
+	if (NewMaterial != nullptr)
+	{
+		NewMaterial->PostEditChange();
+	}
 
 	EmitGeneratedMeshAsset(Options, ResultsOut, &FinalMesh, NewMaterial, WriteDebugMesh);
 	ResultsOut.ResultCode = EResultCode::Success;
