Changed WorldPartitionHLODUtilities plugin type from Editor to UncookedOnly to fix a error reported by the linker when searching for /Script/WorldPartitionHLODUtilities when running in -game.

diff --git a/Engine/Plugins/Editor/WorldPartitionHLODUtilities/WorldPartitionHLODUtilities.uplugin b/Engine/Plugins/Editor/WorldPartitionHLODUtilities/WorldPartitionHLODUtilities.uplugin
--- a/Engine/Plugins/Editor/WorldPartitionHLODUtilities/WorldPartitionHLODUtilities.uplugin
+++ b/Engine/Plugins/Editor/WorldPartitionHLODUtilities/WorldPartitionHLODUtilities.uplugin
@@ -15,7 +15,7 @@
 	[
 		{
 			"Name" : "WorldPartitionHLODUtilities",
-			"Type" : "Editor",
+			"Type" : "UncookedOnly",
 			"LoadingPhase" : "Default"
 		}
 	]
