World Partition HLOD Utilities:
- Added a new HLOD creation parameter to force generated HLOD actors to be in a specific runtime grid instead of getting the target grid from the HLODLayer object.
- Added a log that outputs the reason why an existing HLOD actor needs to be recomputed.

diff --git a/Engine/Plugins/Editor/WorldPartitionHLODUtilities/Source/Private/WorldPartition/HLOD/Utilities/WorldPartitionHLODUtilities.cpp b/Engine/Plugins/Editor/WorldPartitionHLODUtilities/Source/Private/WorldPartition/HLOD/Utilities/WorldPartitionHLODUtilities.cpp
--- a/Engine/Plugins/Editor/WorldPartitionHLODUtilities/Source/Private/WorldPartition/HLOD/Utilities/WorldPartitionHLODUtilities.cpp
+++ b/Engine/Plugins/Editor/WorldPartitionHLODUtilities/Source/Private/WorldPartition/HLOD/Utilities/WorldPartitionHLODUtilities.cpp
@@ -182,7 +182,7 @@ TArray<AWorldPartitionHLOD*> FWorldPartitionHLODUtilities::CreateHLODActors(FHLO
 			check(HLODActor->GetClass() == HLODLayer->GetHLODActorClass());
 		}
 
-		bool bIsDirty = false;
+		const TCHAR* DirtyReason = nullptr;
 
 		// Source actors object
 		UWorldPartitionHLODSourceActorsFromCell* HLODSourceActors = Cast<UWorldPartitionHLODSourceActorsFromCell>(HLODActor->GetSourceActors());
@@ -191,7 +191,7 @@ TArray<AWorldPartitionHLOD*> FWorldPartitionHLODUtilities::CreateHLODActors(FHLO
 			HLODSourceActors = NewObject<UWorldPartitionHLODSourceActorsFromCell>(HLODActor);
 			HLODSourceActors->SetHLODLayer(HLODLayer);
 			HLODActor->SetSourceActors(HLODSourceActors);
-			bIsDirty = true;
+			DirtyReason = TEXT("SourceActors");
 		}
 		check(HLODSourceActors->GetHLODLayer() == HLODLayer);
 
@@ -210,37 +210,37 @@ TArray<AWorldPartitionHLOD*> FWorldPartitionHLODUtilities::CreateHLODActors(FHLO
 			if (bSubActorsChanged)
 			{
 				HLODSourceActors->SetActors(SubActors.Array());
-				bIsDirty = true;
+				DirtyReason = TEXT("SubActors");
 			}
 		}
 
 		// Runtime grid
-		FName RuntimeGrid = HLODLayer->GetRuntimeGrid(InCreationParams.HLODLevel);
+		const FName RuntimeGrid = InCreationParams.RuntimeGrid.IsNone() ? HLODLayer->GetRuntimeGrid(InCreationParams.HLODLevel) : InCreationParams.RuntimeGrid;
 		if (HLODActor->GetRuntimeGrid() != RuntimeGrid)
 		{
 			HLODActor->SetRuntimeGrid(RuntimeGrid);
-			bIsDirty = true;
+			DirtyReason = TEXT("RuntimeGrid");
 		}
 
 		// Spatially loaded
 		if (HLODActor->GetIsSpatiallyLoaded() != HLODLayer->IsSpatiallyLoaded())
 		{
 			HLODActor->SetIsSpatiallyLoaded(HLODLayer->IsSpatiallyLoaded());
-			bIsDirty = true;
+			DirtyReason = TEXT("SpatiallyLoaded");
 		}
 
 		// HLOD level
 		if (HLODActor->GetLODLevel() != InCreationParams.HLODLevel)
 		{
 			HLODActor->SetLODLevel(InCreationParams.HLODLevel);
-			bIsDirty = true;
+			DirtyReason = TEXT("HLODLevel");
 		}
 
 		// Require warmup
 		if (HLODActor->DoesRequireWarmup() != HLODLayer->DoesRequireWarmup())
 		{
 			HLODActor->SetRequireWarmup(HLODLayer->DoesRequireWarmup());
-			bIsDirty = true;
+			DirtyReason = TEXT("DoesRequireWarmup");
 		}
 
 		// Parent HLOD layer
@@ -248,7 +248,7 @@ TArray<AWorldPartitionHLOD*> FWorldPartitionHLODUtilities::CreateHLODActors(FHLO
 		if (HLODActor->GetHLODLayer() != ParentHLODLayer)
 		{
 			HLODActor->SetHLODLayer(ParentHLODLayer);
-			bIsDirty = true;
+			DirtyReason = TEXT("ParentHLODLayer");
 		}
 
 		// Actor label
@@ -256,7 +256,7 @@ TArray<AWorldPartitionHLOD*> FWorldPartitionHLODUtilities::CreateHLODActors(FHLO
 		if (HLODActor->GetActorLabel() != ActorLabel)
 		{
 			HLODActor->SetActorLabel(ActorLabel);
-			bIsDirty = true;
+			DirtyReason = TEXT("ActorLabel");
 		}
 
 		// Folder name
@@ -264,27 +264,28 @@ TArray<AWorldPartitionHLOD*> FWorldPartitionHLODUtilities::CreateHLODActors(FHLO
 		if (HLODActor->GetFolderPath() != FolderPath)
 		{
 			HLODActor->SetFolderPath(FolderPath);
-			bIsDirty = true;
+			DirtyReason = TEXT("FolderPath");
 		}
 
 		// Cell bounds
 		if (!HLODActor->GetHLODBounds().Equals(InCreationParams.CellBounds))
 		{
 			HLODActor->SetHLODBounds(InCreationParams.CellBounds);
-			bIsDirty = true;
+			DirtyReason = TEXT("CellBounds");
 		}
 
 		// Minimum visible distance
 		if (!FMath::IsNearlyEqual(HLODActor->GetMinVisibleDistance(), InCreationParams.MinVisibleDistance))
 		{
 			HLODActor->SetMinVisibleDistance(InCreationParams.MinVisibleDistance);
-			bIsDirty = true;
+			DirtyReason = TEXT("MinVisibleDistance");
 		}
 
 		// If any change was performed, mark HLOD package as dirty
-		if (bIsDirty)
+		if (DirtyReason)
 		{
 			HLODActor->MarkPackageDirty();
+			UE_CLOG(!bNewActor, LogHLODBuilder, Log, TEXT("Marking existing HLOD actor \"%s\" dirty, reason \"%s\"."), *HLODActor->GetActorLabel(), DirtyReason);
 		}
 
 		HLODActors.Add(HLODActor);
diff --git a/Engine/Source/Runtime/Engine/Public/WorldPartition/HLOD/IWorldPartitionHLODUtilities.h b/Engine/Source/Runtime/Engine/Public/WorldPartition/HLOD/IWorldPartitionHLODUtilities.h
--- a/Engine/Source/Runtime/Engine/Public/WorldPartition/HLOD/IWorldPartitionHLODUtilities.h
+++ b/Engine/Source/Runtime/Engine/Public/WorldPartition/HLOD/IWorldPartitionHLODUtilities.h
@@ -29,7 +29,8 @@ struct FHLODCreationParams
 
 	FGuid CellGuid;
 	FString CellName;
-	FBox  CellBounds;
+	FBox CellBounds;
+	FName RuntimeGrid;
 	uint32 HLODLevel;
 	FGuid ContentBundleGuid;
 	TArray<const UDataLayerInstance*> DataLayerInstances;
