World Partition - HLOD: Include editor only references in the list of source actors
* When loading source actors to build HLODs, this prevent errors as editor only references were missing from the instancing context
* These actors can be safely included as "source actors" as they are not HLOD relevant and will be discarded from the build.
* HLOD actors now store their source actors using FWorldPartitionRuntimeCellObjectMapping rather than a custom FHLODSubActor class

diff --git a/Engine/Config/BaseEngine.ini b/Engine/Config/BaseEngine.ini
--- a/Engine/Config/BaseEngine.ini
+++ b/Engine/Config/BaseEngine.ini
@@ -1221,6 +1221,12 @@ DefaultWorldPartitionHLODLayer=/Engine/Maps/Templates/HLODs/HLODLayer_Instanced.
 
 +EnumRedirects=(OldName="ETargetingTraceType", ValueChanges=(("Sweep", "Sphere")) )
 
+; HLOD actors now store their source actors using FWorldPartitionRuntimeCellObjectMapping rather than a custom FHLODSubActor class
++StructRedirects=(OldName="/Script/Engine.HLODSubActor",							NewName="/Script/Engine.WorldPartitionRuntimeCellObjectMapping")
++PropertyRedirects=(OldName="WorldPartitionRuntimeCellObjectMapping.ActorPackage",	NewName="Package")
++PropertyRedirects=(OldName="WorldPartitionRuntimeCellObjectMapping.ActorPath",		NewName="Path")
++PropertyRedirects=(OldName="WorldPartitionRuntimeCellObjectMapping.ActorGuid",		NewName="ActorInstanceGuid")
+
 ; End of releases redirects
 
 
diff --git a/Engine/Plugins/Editor/WorldPartitionHLODUtilities/Source/Private/WorldPartition/HLOD/Utilities/WorldPartitionHLODUtilities.cpp b/Engine/Plugins/Editor/WorldPartitionHLODUtilities/Source/Private/WorldPartition/HLOD/Utilities/WorldPartitionHLODUtilities.cpp
--- a/Engine/Plugins/Editor/WorldPartitionHLODUtilities/Source/Private/WorldPartition/HLOD/Utilities/WorldPartitionHLODUtilities.cpp
+++ b/Engine/Plugins/Editor/WorldPartitionHLODUtilities/Source/Private/WorldPartition/HLOD/Utilities/WorldPartitionHLODUtilities.cpp
@@ -14,7 +14,6 @@
 #include "WorldPartition/HLOD/HLODModifier.h"
 #include "WorldPartition/HLOD/HLODSourceActorsFromCell.h"
 #include "WorldPartition/HLOD/HLODStats.h"
-#include "WorldPartition/HLOD/HLODSubActor.h"
 #include "WorldPartition/HLOD/HLODInstancedStaticMeshComponent.h"
 #include "WorldPartition/HLOD/Builders/HLODBuilderInstancing.h"
 #include "WorldPartition/HLOD/Builders/HLODBuilderMeshMerge.h"
@@ -79,31 +78,66 @@ static uint32 ComputeHLODHash(AWorldPartitionHLOD* InHLODActor, const TArray<UAc
 	return Ar.GetCrc();
 }
 
-void AddSubActor(const FWorldPartitionActorDescView& ActorDescView, const IStreamingGenerationContext::FActorInstance& ActorInstance, TSet<FHLODSubActor>& SubActors)
+void AddSourceActor(const FWorldPartitionActorDescView& ActorDescView, const IStreamingGenerationContext::FActorInstance& ActorInstance, const FName WorldPackageName, TMap<FGuid, FWorldPartitionRuntimeCellObjectMapping>& SourceActors)
 {
 	const FName ActorPath = *ActorDescView.GetActorSoftPath().ToString();
 
-	// Add the actor
-	bool bIsAlreadyInSet = false;
 	const UActorDescContainer* ActorDescContainer = ActorDescView.GetActorDesc()->GetContainer();
-	FHLODSubActor SubActor(ActorDescView.GetGuid(), ActorDescView.GetActorPackage(), ActorPath, ActorInstance.GetContainerID(), ActorDescContainer->GetContainerPackage(), ActorInstance.GetTransform());
-	SubActors.Add(SubActor, &bIsAlreadyInSet);
 
-	if (!bIsAlreadyInSet)
+	FGuid ActorInstanceGuid = ActorInstance.GetContainerID().GetActorGuid(ActorDescView.GetGuid());	
+	FWorldPartitionRuntimeCellObjectMapping& ActorMapping = SourceActors.FindOrAdd(ActorInstanceGuid);
+	if (!ActorMapping.ActorInstanceGuid.IsValid())
 	{
-		// Add its references
+		ActorMapping = FWorldPartitionRuntimeCellObjectMapping(
+			ActorDescView.GetActorPackage(),
+			*ActorDescView.GetActorSoftPath().ToString(),
+			ActorDescView.GetBaseClass(),
+			ActorDescView.GetNativeClass(),
+			ActorInstance.GetContainerID(),
+			ActorInstance.GetTransform(),
+			ActorDescContainer->GetContainerPackage(),
+			WorldPackageName,
+			ActorInstanceGuid,
+			false);
+
+		// Add its runtime references, recursively
 		const FActorDescViewMap* ActorDescViewMap = ActorInstance.ActorSetInstance->ContainerInstance->ActorDescViewMap;
 		for (const FGuid& ReferenceGuid : ActorDescView.GetReferences())
 		{
 			const FWorldPartitionActorDescView& RefActorDescView = ActorDescViewMap->FindByGuidChecked(ReferenceGuid);
-			AddSubActor(RefActorDescView, ActorInstance, SubActors);
+			AddSourceActor(RefActorDescView, ActorInstance, WorldPackageName, SourceActors);
+		}
+
+		// Add its editor references
+		for (const FGuid& EditorReferenceGuid : ActorDescView.GetEditorReferences())
+		{
+			const FWorldPartitionActorDesc& ReferenceActorDesc = ActorDescContainer->GetActorDescChecked(EditorReferenceGuid);
+			FGuid EditorRefInstanceGuid = ActorInstance.GetContainerID().GetActorGuid(EditorReferenceGuid);
+			FWorldPartitionRuntimeCellObjectMapping& EditorRefMapping = SourceActors.FindOrAdd(EditorRefInstanceGuid);
+			if (!EditorRefMapping.ActorInstanceGuid.IsValid())
+			{
+				EditorRefMapping = FWorldPartitionRuntimeCellObjectMapping(
+					ReferenceActorDesc.GetActorPackage(),
+					*ReferenceActorDesc.GetActorSoftPath().ToString(),
+					ReferenceActorDesc.GetBaseClass(),
+					ReferenceActorDesc.GetNativeClass(),
+					ActorInstance.GetContainerID(),
+					ActorInstance.GetTransform(),
+					ActorDescContainer->GetContainerPackage(),
+					WorldPackageName,
+					EditorRefInstanceGuid,
+					true
+				);
+			}
 		}
 	}
 }
 
 TArray<AWorldPartitionHLOD*> FWorldPartitionHLODUtilities::CreateHLODActors(FHLODCreationContext& InCreationContext, const FHLODCreationParams& InCreationParams, const TArray<IStreamingGenerationContext::FActorInstance>& InActors)
 {
-	TMap<UHLODLayer*, TSet<FHLODSubActor>> SubActorsPerHLODLayer;
+	TMap<UHLODLayer*, TMap<FGuid, FWorldPartitionRuntimeCellObjectMapping>> SourceActorsPerHLODLayer;
+
+	FName WorldPackageName = InCreationParams.WorldPartition->GetWorld()->GetPackage()->GetFName();
 
 	for (const IStreamingGenerationContext::FActorInstance& ActorInstance : InActors)
 	{
@@ -119,18 +153,18 @@ TArray<AWorldPartitionHLOD*> FWorldPartitionHLODUtilities::CreateHLODActors(FHLO
 
 			if (UHLODLayer* HLODLayer = Cast<UHLODLayer>(ActorDescView.GetHLODLayer().TryLoad()))
 			{
-				TSet<FHLODSubActor>& SubActors = SubActorsPerHLODLayer.FindOrAdd(HLODLayer);
-				AddSubActor(ActorDescView, ActorInstance, SubActors);
+				TMap<FGuid, FWorldPartitionRuntimeCellObjectMapping>& SubActors = SourceActorsPerHLODLayer.FindOrAdd(HLODLayer);
+				AddSourceActor(ActorDescView, ActorInstance, WorldPackageName, SubActors);
 			}
 		}
 	}
 
 	TArray<AWorldPartitionHLOD*> HLODActors;
-	for (const auto& Pair : SubActorsPerHLODLayer)
+	for (auto& Pair : SourceActorsPerHLODLayer)
 	{
 		const UHLODLayer* HLODLayer = Pair.Key;
-		const TSet<FHLODSubActor>& SubActors = Pair.Value;
-		check(!SubActors.IsEmpty());
+		TMap<FGuid, FWorldPartitionRuntimeCellObjectMapping>& SourceActorsMap = Pair.Value;
+		check(!SourceActorsMap.IsEmpty());
 
 		auto ComputeHLODActorUniqueHash = [](const UHLODLayer* HLODLayer, const FGuid CellGuid)
 		{
@@ -197,22 +231,24 @@ TArray<AWorldPartitionHLOD*> FWorldPartitionHLODUtilities::CreateHLODActors(FHLO
 		}
 		check(HLODSourceActors->GetHLODLayer() == HLODLayer);
 
-		// Sub actors
+		// Assign source actors
 		{
-			bool bSubActorsChanged = HLODSourceActors->GetActors().Num() != SubActors.Num();
-			if (!bSubActorsChanged)
+			TArray<FWorldPartitionRuntimeCellObjectMapping> SourceActorsArray;
+			SourceActorsMap.KeySort(TLess<FGuid>());
+			SourceActorsMap.GenerateValueArray(SourceActorsArray);
+
+			bool bSourceActorsChanged = HLODSourceActors->GetActors().Num() != SourceActorsArray.Num();
+			if (!bSourceActorsChanged)
 			{
-				TArray<FHLODSubActor> A = HLODSourceActors->GetActors();
-				TArray<FHLODSubActor> B = SubActors.Array();
-				A.Sort();
-				B.Sort();
-				bSubActorsChanged = A != B;
+				const uint32 NewHash = UWorldPartitionHLODSourceActorsFromCell::GetHLODHash(SourceActorsArray);
+				const uint32 PreviousHash = UWorldPartitionHLODSourceActorsFromCell::GetHLODHash(HLODSourceActors->GetActors());
+				bSourceActorsChanged = NewHash != PreviousHash;
 			}
 
-			if (bSubActorsChanged)
+			if (bSourceActorsChanged)
 			{
-				HLODSourceActors->SetActors(SubActors.Array());
-				DirtyReason = TEXT("SubActors");
+				HLODSourceActors->SetActors(MoveTemp(SourceActorsArray));
+				DirtyReason = TEXT("SourceActors");
 			}
 		}
 
diff --git a/Engine/Source/Runtime/Engine/Private/WorldPartition/HLOD/HLODActor.cpp b/Engine/Source/Runtime/Engine/Private/WorldPartition/HLOD/HLODActor.cpp
--- a/Engine/Source/Runtime/Engine/Private/WorldPartition/HLOD/HLODActor.cpp
+++ b/Engine/Source/Runtime/Engine/Private/WorldPartition/HLOD/HLODActor.cpp
@@ -136,7 +136,7 @@ void AWorldPartitionHLOD::Serialize(FArchive& Ar)
 		{
 			check(!SourceActors)
 			UWorldPartitionHLODSourceActorsFromCell* SourceActorsFromCell = NewObject<UWorldPartitionHLODSourceActorsFromCell>(this);
-			SourceActorsFromCell->SetActors(HLODSubActors_DEPRECATED);
+			SourceActorsFromCell->SetActors(MoveTemp(HLODSubActors_DEPRECATED));
 			SourceActorsFromCell->SetHLODLayer(SubActorsHLODLayer_DEPRECATED);
 			SourceActors = SourceActorsFromCell;
 		}
diff --git a/Engine/Source/Runtime/Engine/Private/WorldPartition/HLOD/HLODActorDesc.cpp b/Engine/Source/Runtime/Engine/Private/WorldPartition/HLOD/HLODActorDesc.cpp
--- a/Engine/Source/Runtime/Engine/Private/WorldPartition/HLOD/HLODActorDesc.cpp
+++ b/Engine/Source/Runtime/Engine/Private/WorldPartition/HLOD/HLODActorDesc.cpp
@@ -34,14 +34,14 @@ void FHLODActorDesc::Init(const AActor* InActor)
 
 	if (HLODCellSourceActors && WorldPartition)
 	{
-		for (const FHLODSubActor& SubActor : HLODCellSourceActors->GetActors())
+		for (const FWorldPartitionRuntimeCellObjectMapping& SubActor : HLODCellSourceActors->GetActors())
 		{
 			if (SubActor.ContainerID.IsMainContainer())
 			{
-				const FWorldPartitionActorDesc* SubActorDesc = WorldPartition->GetActorDesc(SubActor.ActorGuid);
+				const FWorldPartitionActorDesc* SubActorDesc = WorldPartition->GetActorDesc(SubActor.ActorInstanceGuid);
 				if (SubActorDesc && SubActorDesc->GetActorNativeClass()->IsChildOf<AWorldPartitionHLOD>())
 				{
-					ChildHLODActors.Add(SubActor.ActorGuid);
+					ChildHLODActors.Add(SubActor.ActorInstanceGuid);
 				}
 			}
 		}
diff --git a/Engine/Source/Runtime/Engine/Private/WorldPartition/HLOD/HLODSourceActorsFromCell.cpp b/Engine/Source/Runtime/Engine/Private/WorldPartition/HLOD/HLODSourceActorsFromCell.cpp
--- a/Engine/Source/Runtime/Engine/Private/WorldPartition/HLOD/HLODSourceActorsFromCell.cpp
+++ b/Engine/Source/Runtime/Engine/Private/WorldPartition/HLOD/HLODSourceActorsFromCell.cpp
@@ -4,26 +4,10 @@
 
 #if WITH_EDITOR
 #include "WorldPartition/HLOD/HLODActor.h"
+#include "Serialization/ArchiveCrc32.h"
 #include "WorldPartition/WorldPartitionLevelStreamingDynamic.h"
 #endif
 
-
-#if WITH_EDITOR
-
-static uint32 GetHLODHash(const TArray<FHLODSubActor>& InSubActors)
-{
-	uint32 HLODHash = 0;
-
-	for (const FHLODSubActor& HLODSubActor : InSubActors)
-	{
-		HLODHash = HashCombine(GetTypeHash(HLODSubActor), HLODHash);
-	}
-
-	return HLODHash;
-}
-
-#endif // #if WITH_EDITOR
-
 UWorldPartitionHLODSourceActorsFromCell::UWorldPartitionHLODSourceActorsFromCell(const FObjectInitializer& ObjectInitializer)
 	: Super(ObjectInitializer)
 {
@@ -38,30 +22,10 @@ ULevelStreaming* UWorldPartitionHLODSourceActorsFromCell::LoadSourceActors(bool&
 	bOutDirty = false;
 	AWorldPartitionHLOD* HLODActor = CastChecked<AWorldPartitionHLOD>(GetOuter());
 	UWorld* World = HLODActor->GetWorld();
-	UWorldPartition* WorldPartition = World->GetWorldPartition();
-	check(WorldPartition);
 
-	const FTopLevelAssetPath WorldPartitionHLODNativeClass(AWorldPartitionHLOD::StaticClass()->GetPathName());
 	const FName LevelStreamingName = FName(*FString::Printf(TEXT("HLODLevelStreaming_%s"), *HLODActor->GetName()));
-	TArray<FWorldPartitionRuntimeCellObjectMapping> Mappings;
-	Mappings.Reserve(Actors.Num());
-	Algo::Transform(Actors, Mappings, [World, HLODActor, &WorldPartitionHLODNativeClass](const FHLODSubActor& SubActor)
-	{		
-		return FWorldPartitionRuntimeCellObjectMapping(
-			SubActor.ActorPackage,
-			SubActor.ActorPath,
-			FTopLevelAssetPath(),
-			WorldPartitionHLODNativeClass,
-			SubActor.ContainerID,
-			SubActor.ContainerTransform,
-			SubActor.ContainerPackage,
-			World->GetPackage()->GetFName(),
-			SubActor.ActorGuid,
-			false
-		);
-	});
-
-	UWorldPartitionLevelStreamingDynamic* LevelStreaming = UWorldPartitionLevelStreamingDynamic::LoadInEditor(World, LevelStreamingName, Mappings);
+
+	UWorldPartitionLevelStreamingDynamic* LevelStreaming = UWorldPartitionLevelStreamingDynamic::LoadInEditor(World, LevelStreamingName, Actors);
 	check(LevelStreaming);
 
 	if (!LevelStreaming->GetLoadSucceeded())
@@ -72,24 +36,33 @@ ULevelStreaming* UWorldPartitionHLODSourceActorsFromCell::LoadSourceActors(bool&
 	return LevelStreaming;
 }
 
+uint32 UWorldPartitionHLODSourceActorsFromCell::GetHLODHash(const TArray<FWorldPartitionRuntimeCellObjectMapping>& InSourceActors)
+{
+	TArray<FWorldPartitionRuntimeCellObjectMapping>& MutableSourceActors(const_cast<TArray<FWorldPartitionRuntimeCellObjectMapping>&>(InSourceActors));
+
+	FArchiveCrc32 Ar;
+	Ar << MutableSourceActors;
+	return Ar.GetCrc();
+}
+
 uint32 UWorldPartitionHLODSourceActorsFromCell::GetHLODHash() const
 {
 	uint32 HLODHash = Super::GetHLODHash();
 
-	// Sub Actors
-	uint32 SubActorsHash = ::GetHLODHash(Actors);
-	UE_LOG(LogHLODHash, VeryVerbose, TEXT(" - Sub Actors (%d actors) = %x"), Actors.Num(), SubActorsHash);
-	HLODHash = HashCombine(HLODHash, SubActorsHash);
+	// Source Actors
+	uint32 SourceActorsHash = GetHLODHash(Actors);
+	UE_LOG(LogHLODHash, VeryVerbose, TEXT(" - Source Actors (%d actors) = %x"), Actors.Num(), SourceActorsHash);
+	HLODHash = HashCombine(HLODHash, SourceActorsHash);
 
 	return HLODHash;
 }
 
-void UWorldPartitionHLODSourceActorsFromCell::SetActors(const TArray<FHLODSubActor>& InSourceActors)
+void UWorldPartitionHLODSourceActorsFromCell::SetActors(const TArray<FWorldPartitionRuntimeCellObjectMapping>&& InSourceActors)
 {
 	Actors = InSourceActors;
 }
 
-const TArray<FHLODSubActor>& UWorldPartitionHLODSourceActorsFromCell::GetActors() const
+const TArray<FWorldPartitionRuntimeCellObjectMapping>& UWorldPartitionHLODSourceActorsFromCell::GetActors() const
 {
 	return Actors;
 }
diff --git a/Engine/Source/Runtime/Engine/Private/WorldPartition/WorldPartitionRuntimeLevelStreamingCell.cpp b/Engine/Source/Runtime/Engine/Private/WorldPartition/WorldPartitionRuntimeLevelStreamingCell.cpp
--- a/Engine/Source/Runtime/Engine/Private/WorldPartition/WorldPartitionRuntimeLevelStreamingCell.cpp
+++ b/Engine/Source/Runtime/Engine/Private/WorldPartition/WorldPartitionRuntimeLevelStreamingCell.cpp
@@ -260,8 +260,8 @@ void UWorldPartitionRuntimeLevelStreamingCell::AddActorToCell(const FWorldPartit
 		Packages.Emplace(
 			ReferenceActorDesc.GetActorPackage(), 
 			*ReferenceActorDesc.GetActorSoftPath().ToString(), 
-			ActorDescView.GetBaseClass(),
-			ActorDescView.GetNativeClass(),
+			ReferenceActorDesc.GetBaseClass(),
+			ReferenceActorDesc.GetNativeClass(),
 			InContainerID, 
 			InContainerTransform, 
 			InContainer->GetContainerPackage(), 
diff --git a/Engine/Source/Runtime/Engine/Public/WorldPartition/HLOD/HLODActor.h b/Engine/Source/Runtime/Engine/Public/WorldPartition/HLOD/HLODActor.h
--- a/Engine/Source/Runtime/Engine/Public/WorldPartition/HLOD/HLODActor.h
+++ b/Engine/Source/Runtime/Engine/Public/WorldPartition/HLOD/HLODActor.h
@@ -7,7 +7,6 @@
 #include "Containers/Set.h"
 #include "WorldPartition/WorldPartitionRuntimeCell.h"
 #include "WorldPartition/WorldPartitionActorDesc.h"
-#include "WorldPartition/HLOD/HLODSubActor.h"
 #include "WorldPartition/HLOD/HLODStats.h"
 #include "HLODActor.generated.h"
 
@@ -126,7 +125,7 @@ private:
 	FName SourceCellName_DEPRECATED;
 
 	UPROPERTY()
-	TArray<FHLODSubActor> HLODSubActors_DEPRECATED;
+	TArray<FWorldPartitionRuntimeCellObjectMapping> HLODSubActors_DEPRECATED;
 
 	UPROPERTY()
 	TObjectPtr<const UHLODLayer> SubActorsHLODLayer_DEPRECATED;
diff --git a/Engine/Source/Runtime/Engine/Public/WorldPartition/HLOD/HLODActorDesc.h b/Engine/Source/Runtime/Engine/Public/WorldPartition/HLOD/HLODActorDesc.h
--- a/Engine/Source/Runtime/Engine/Public/WorldPartition/HLOD/HLODActorDesc.h
+++ b/Engine/Source/Runtime/Engine/Public/WorldPartition/HLOD/HLODActorDesc.h
@@ -10,7 +10,6 @@
 #include "WorldPartition/WorldPartitionActorDesc.h"
 #include "WorldPartition/DataLayer/DataLayersID.h"
 #include "WorldPartition/HLOD/HLODStats.h"
-#include "WorldPartition/HLOD/HLODSubActor.h"
 
 class UHLODLayer;
 class AWorldPartitionHLOD;
diff --git a/Engine/Source/Runtime/Engine/Public/WorldPartition/HLOD/HLODSourceActorsFromCell.h b/Engine/Source/Runtime/Engine/Public/WorldPartition/HLOD/HLODSourceActorsFromCell.h
--- a/Engine/Source/Runtime/Engine/Public/WorldPartition/HLOD/HLODSourceActorsFromCell.h
+++ b/Engine/Source/Runtime/Engine/Public/WorldPartition/HLOD/HLODSourceActorsFromCell.h
@@ -3,7 +3,7 @@
 #pragma once
 
 #include "WorldPartition/HLOD/HLODSourceActors.h"
-#include "WorldPartition/HLOD/HLODSubActor.h"
+#include "WorldPartition/WorldPartitionRuntimeCell.h"
 #include "HLODSourceActorsFromCell.generated.h"
 
 
@@ -17,13 +17,15 @@ public:
 	ENGINE_API virtual ULevelStreaming* LoadSourceActors(bool& bOutDirty) const override;
 	ENGINE_API virtual uint32 GetHLODHash() const override;
 
-	ENGINE_API void SetActors(const TArray<FHLODSubActor>& InSourceActors);
-	ENGINE_API const TArray<FHLODSubActor>& GetActors() const;
+	ENGINE_API static uint32 GetHLODHash(const TArray<FWorldPartitionRuntimeCellObjectMapping>& InSourceActors);
+
+	ENGINE_API void SetActors(const TArray<FWorldPartitionRuntimeCellObjectMapping>&& InSourceActors);
+	ENGINE_API const TArray<FWorldPartitionRuntimeCellObjectMapping>& GetActors() const;
 #endif
 
 private:
 #if WITH_EDITORONLY_DATA
 	UPROPERTY()
-	TArray<FHLODSubActor> Actors;
+	TArray<FWorldPartitionRuntimeCellObjectMapping> Actors;
 #endif
 };
diff --git a/Engine/Source/Runtime/Engine/Public/WorldPartition/HLOD/HLODSubActor.h b/Engine/Source/Runtime/Engine/Public/WorldPartition/HLOD/HLODSubActor.h
deleted file mode 100644
