World Partition - HLOD: Emissive texture map is not assigned to HLOD material instance when building HLODs with Approximate Mesh

diff --git a/Engine/Plugins/Editor/WorldPartitionHLODUtilities/Source/Private/WorldPartition/HLOD/Builders/HLODBuilderMeshApproximate.cpp b/Engine/Plugins/Editor/WorldPartitionHLODUtilities/Source/Private/WorldPartition/HLOD/Builders/HLODBuilderMeshApproximate.cpp
--- a/Engine/Plugins/Editor/WorldPartitionHLODUtilities/Source/Private/WorldPartition/HLOD/Builders/HLODBuilderMeshApproximate.cpp
+++ b/Engine/Plugins/Editor/WorldPartitionHLODUtilities/Source/Private/WorldPartition/HLOD/Builders/HLODBuilderMeshApproximate.cpp
@@ -46,6 +46,10 @@ uint32 UHLODBuilderMeshApproximateSettings::GetCRC() const
 
 	FArchiveCrc32 Ar;
 
+	// Base key, changing this will force a rebuild of all HLODs from this builder
+	FString HLODBaseKey = "1EC5FBC75A71412EB296F0E7E8424814";
+	Ar << HLODBaseKey;
+
 	Ar << This.MeshApproximationSettings;
 	UE_LOG(LogHLODBuilder, VeryVerbose, TEXT(" - MeshApproximationSettings = %d"), Ar.GetCrc());
 
@@ -106,7 +110,7 @@ TArray<UActorComponent*> UHLODBuilderMeshApproximate::Build(const FHLODBuildCont
 	Options.MetallicTexParamName = FName(FMaterialUtilities::GetFlattenMaterialTextureName(EFlattenMaterialProperties::Metallic, Options.BakeMaterial));
 	Options.RoughnessTexParamName = FName(FMaterialUtilities::GetFlattenMaterialTextureName(EFlattenMaterialProperties::Roughness, Options.BakeMaterial));
 	Options.SpecularTexParamName = FName(FMaterialUtilities::GetFlattenMaterialTextureName(EFlattenMaterialProperties::Specular, Options.BakeMaterial));
-	Options.EmissiveTexParamName = FName(FMaterialUtilities::GetFlattenMaterialTextureName(EFlattenMaterialProperties::Emissive, Options.BakeMaterial));
+	Options.EmissiveTexParamName = FName("EmissiveHDRTexture"); // TODO - Approximate actors should look up if the material sampler is expecting an HDR texture and capture accordingly
 	Options.bUsePackedMRS = true;
 	Options.PackedMRSTexParamName = FName("PackedTexture");
 
@@ -194,6 +198,7 @@ TArray<UActorComponent*> UHLODBuilderMeshApproximate::Build(const FHLODBuildCont
 			SetStaticSwitch("UseMetallic", Options.bBakeMetallic);
 			SetStaticSwitch("UseSpecular", Options.bBakeSpecular);
 			SetStaticSwitch("UseEmissive", Options.bBakeEmissive);
+			SetStaticSwitch("UseEmissiveColor", Options.bBakeEmissive);
 			SetStaticSwitch("UseEmissiveHDR", Options.bBakeEmissive);
 			SetStaticSwitch("UseNormal", Options.bBakeNormalMap);
 			SetStaticSwitch("PackMetallic", Options.bUsePackedMRS);
diff --git a/Engine/Source/Developer/HierarchicalLODUtilities/Private/HierarchicalLODUtilities.cpp b/Engine/Source/Developer/HierarchicalLODUtilities/Private/HierarchicalLODUtilities.cpp
--- a/Engine/Source/Developer/HierarchicalLODUtilities/Private/HierarchicalLODUtilities.cpp
+++ b/Engine/Source/Developer/HierarchicalLODUtilities/Private/HierarchicalLODUtilities.cpp
@@ -537,7 +537,7 @@ FHLODBuildResults GenerateHLODMesh_Approximate(const FHLODBuildParams& InBuildPa
 	Options.MetallicTexParamName = FName(FMaterialUtilities::GetFlattenMaterialTextureName(EFlattenMaterialProperties::Metallic, Options.BakeMaterial));
 	Options.RoughnessTexParamName = FName(FMaterialUtilities::GetFlattenMaterialTextureName(EFlattenMaterialProperties::Roughness, Options.BakeMaterial));
 	Options.SpecularTexParamName = FName(FMaterialUtilities::GetFlattenMaterialTextureName(EFlattenMaterialProperties::Specular, Options.BakeMaterial));
-	Options.EmissiveTexParamName = FName(FMaterialUtilities::GetFlattenMaterialTextureName(EFlattenMaterialProperties::Emissive, Options.BakeMaterial));
+	Options.EmissiveTexParamName = FName("EmissiveHDRTexture"); // TODO - Approximate actors should look up if the material sampler is expecting an HDR texture and capture accordingly
 	Options.bUsePackedMRS = true;
 	Options.PackedMRSTexParamName = FName("PackedTexture");
 
@@ -619,6 +619,7 @@ FHLODBuildResults GenerateHLODMesh_Approximate(const FHLODBuildParams& InBuildPa
 			SetStaticSwitch("UseMetallic", Options.bBakeMetallic);
 			SetStaticSwitch("UseSpecular", Options.bBakeSpecular);
 			SetStaticSwitch("UseEmissive", Options.bBakeEmissive);
+			SetStaticSwitch("UseEmissiveColor", Options.bBakeEmissive);
 			SetStaticSwitch("UseEmissiveHDR", Options.bBakeEmissive);
 			SetStaticSwitch("UseNormal", Options.bBakeNormalMap);
 			SetStaticSwitch("PackMetallic", Options.bUsePackedMRS);
diff --git a/Engine/Source/Runtime/Engine/Private/HLODProxy.cpp b/Engine/Source/Runtime/Engine/Private/HLODProxy.cpp
--- a/Engine/Source/Runtime/Engine/Private/HLODProxy.cpp
+++ b/Engine/Source/Runtime/Engine/Private/HLODProxy.cpp
@@ -495,7 +495,7 @@ uint32 UHLODProxy::GetCRC(UStaticMeshComponent* InComponent, uint32 InCRC, const
 }
 
 // Key that forms the basis of the HLOD proxy key. Bump this key (i.e. generate a new GUID) when you want to force a rebuild of ALL HLOD proxies
-#define HLOD_PROXY_BASE_KEY		TEXT("498A8FFD64744FA3B2E7C7589581F91A")
+#define HLOD_PROXY_BASE_KEY		TEXT("136F4B1AD66E47808C62C1F9CC87CC1F")
 
 FName UHLODProxy::GenerateKeyForActor(const ALODActor* LODActor, bool bMustUndoLevelTransform)
 {
