World Partition - HLOD: Avoid the creation of reverse index buffers when building HLODs
* Reversed index buffers are used to handle negative scaling transforms. Since HLOD meshes will never be used with such transforms, we can get rid of these buffers and reduce memory usage.

diff --git a/Engine/Plugins/Editor/WorldPartitionHLODUtilities/Source/Private/WorldPartition/HLOD/Builders/HLODBuilderMeshApproximate.cpp b/Engine/Plugins/Editor/WorldPartitionHLODUtilities/Source/Private/WorldPartition/HLOD/Builders/HLODBuilderMeshApproximate.cpp
--- a/Engine/Plugins/Editor/WorldPartitionHLODUtilities/Source/Private/WorldPartition/HLOD/Builders/HLODBuilderMeshApproximate.cpp
+++ b/Engine/Plugins/Editor/WorldPartitionHLODUtilities/Source/Private/WorldPartition/HLOD/Builders/HLODBuilderMeshApproximate.cpp
@@ -47,7 +47,7 @@ uint32 UHLODBuilderMeshApproximateSettings::GetCRC() const
 	FArchiveCrc32 Ar;
 
 	// Base key, changing this will force a rebuild of all HLODs from this builder
-	FString HLODBaseKey = "1EC5FBC75A71412EB296F0E7E8424814";
+	FString HLODBaseKey = "60D800B93F8B43789AB7FC0BF9F9BDC1";
 	Ar << HLODBaseKey;
 
 	Ar << This.MeshApproximationSettings;
@@ -98,6 +98,7 @@ TArray<UActorComponent*> UHLODBuilderMeshApproximate::Build(const FHLODBuildCont
 	Options.BasePackagePath = InHLODBuildContext.AssetsOuter->GetPackage()->GetName();
 	Options.bGenerateLightmapUVs = false;
 	Options.bCreatePhysicsBody = false;
+	Options.bBuildReversedIndexBuffer = false;
 
 	// Material baking settings
 	Options.BakeMaterial = HLODMaterial;
diff --git a/Engine/Plugins/Editor/WorldPartitionHLODUtilities/Source/Private/WorldPartition/HLOD/Builders/HLODBuilderMeshMerge.cpp b/Engine/Plugins/Editor/WorldPartitionHLODUtilities/Source/Private/WorldPartition/HLOD/Builders/HLODBuilderMeshMerge.cpp
--- a/Engine/Plugins/Editor/WorldPartitionHLODUtilities/Source/Private/WorldPartition/HLOD/Builders/HLODBuilderMeshMerge.cpp
+++ b/Engine/Plugins/Editor/WorldPartitionHLODUtilities/Source/Private/WorldPartition/HLOD/Builders/HLODBuilderMeshMerge.cpp
@@ -59,7 +59,7 @@ uint32 UHLODBuilderMeshMergeSettings::GetCRC() const
 	FArchiveCrc32 Ar;
 
 	// Base mesh merge key, changing this will force a rebuild of all HLODs from this builder
-	FString HLODBaseKey = "B8DB9CB1780C4EE1B80A36D9E205AA0F";
+	FString HLODBaseKey = "2126404D867E40B7BC8B59C408FAF8AA";
 	Ar << HLODBaseKey;
 
 	Ar << This.MeshMergeSettings;
diff --git a/Engine/Plugins/Experimental/MeshModelingToolsetExp/Source/GeometryProcessingAdapters/Private/ApproximateActorsImpl.cpp b/Engine/Plugins/Experimental/MeshModelingToolsetExp/Source/GeometryProcessingAdapters/Private/ApproximateActorsImpl.cpp
--- a/Engine/Plugins/Experimental/MeshModelingToolsetExp/Source/GeometryProcessingAdapters/Private/ApproximateActorsImpl.cpp
+++ b/Engine/Plugins/Experimental/MeshModelingToolsetExp/Source/GeometryProcessingAdapters/Private/ApproximateActorsImpl.cpp
@@ -1399,6 +1399,8 @@ UStaticMesh* FApproximateActorsImpl::EmitGeneratedMeshAsset(
 	MeshAssetOptions.bGenerateLightmapUVs = Options.bGenerateLightmapUVs;
 	MeshAssetOptions.bCreatePhysicsBody = Options.bCreatePhysicsBody;
 
+	MeshAssetOptions.bBuildReversedIndexBuffer = Options.bBuildReversedIndexBuffer;
+
 	if (Material)
 	{
 		MeshAssetOptions.AssetMaterials.Add(Material);
diff --git a/Engine/Plugins/Runtime/MeshModelingToolset/Source/ModelingComponentsEditorOnly/Private/AssetUtils/CreateStaticMeshUtil.cpp b/Engine/Plugins/Runtime/MeshModelingToolset/Source/ModelingComponentsEditorOnly/Private/AssetUtils/CreateStaticMeshUtil.cpp
--- a/Engine/Plugins/Runtime/MeshModelingToolset/Source/ModelingComponentsEditorOnly/Private/AssetUtils/CreateStaticMeshUtil.cpp
+++ b/Engine/Plugins/Runtime/MeshModelingToolset/Source/ModelingComponentsEditorOnly/Private/AssetUtils/CreateStaticMeshUtil.cpp
@@ -76,6 +76,7 @@ UE::AssetUtils::ECreateStaticMeshResult UE::AssetUtils::CreateStaticMeshAsset(
 		BuildSettings.bRecomputeTangents = Options.bEnableRecomputeTangents;
 		BuildSettings.bGenerateLightmapUVs = Options.bGenerateLightmapUVs;
 		BuildSettings.bUseFullPrecisionUVs = Options.bUseFullPrecisionUVs;
+		BuildSettings.bBuildReversedIndexBuffer = Options.bBuildReversedIndexBuffer;
 		
 		if (!Options.bAllowDistanceField)
 		{
diff --git a/Engine/Plugins/Runtime/MeshModelingToolset/Source/ModelingComponentsEditorOnly/Public/AssetUtils/CreateStaticMeshUtil.h b/Engine/Plugins/Runtime/MeshModelingToolset/Source/ModelingComponentsEditorOnly/Public/AssetUtils/CreateStaticMeshUtil.h
--- a/Engine/Plugins/Runtime/MeshModelingToolset/Source/ModelingComponentsEditorOnly/Public/AssetUtils/CreateStaticMeshUtil.h
+++ b/Engine/Plugins/Runtime/MeshModelingToolset/Source/ModelingComponentsEditorOnly/Public/AssetUtils/CreateStaticMeshUtil.h
@@ -96,6 +96,9 @@ namespace AssetUtils
 
 		// by default, PostEditChange() will be called to rebuild mesh, set true to skip this call
 		bool bDeferPostEditChange = false;
+
+		// Whether to build a reversed index buffer. Will double the memory usage of indices.
+		bool bBuildReversedIndexBuffer = true;
 	};
 
 	/**
diff --git a/Engine/Source/Developer/GeometryProcessingInterfaces/Public/GeometryProcessingInterfaces/ApproximateActors.h b/Engine/Source/Developer/GeometryProcessingInterfaces/Public/GeometryProcessingInterfaces/ApproximateActors.h
--- a/Engine/Source/Developer/GeometryProcessingInterfaces/Public/GeometryProcessingInterfaces/ApproximateActors.h
+++ b/Engine/Source/Developer/GeometryProcessingInterfaces/Public/GeometryProcessingInterfaces/ApproximateActors.h
@@ -253,6 +253,8 @@ public:
 		// Whether the UBodySetup on the Asset will be created
 		bool bCreatePhysicsBody = true;
 
+		/** Required to optimize mesh in mirrored transform. Double index buffer size. */
+		bool bBuildReversedIndexBuffer = true;
 
 		//
 		// Debug settings
diff --git a/Engine/Source/Developer/MeshMergeUtilities/Private/MeshMergeUtilities.cpp b/Engine/Source/Developer/MeshMergeUtilities/Private/MeshMergeUtilities.cpp
--- a/Engine/Source/Developer/MeshMergeUtilities/Private/MeshMergeUtilities.cpp
+++ b/Engine/Source/Developer/MeshMergeUtilities/Private/MeshMergeUtilities.cpp
@@ -2569,6 +2569,7 @@ void FMeshMergeUtilities::MergeComponentsToStaticMesh(const TArray<UPrimitiveCom
 				SrcModel.BuildSettings.bRemoveDegenerates = false;
 				SrcModel.BuildSettings.bUseHighPrecisionTangentBasis = false;
 				SrcModel.BuildSettings.bUseFullPrecisionUVs = false;
+				SrcModel.BuildSettings.bBuildReversedIndexBuffer = false;
 				SrcModel.BuildSettings.bGenerateLightmapUVs = LightMapUVChannel != INDEX_NONE;
 				SrcModel.BuildSettings.MinLightmapResolution = InSettings.bComputedLightMapResolution ? DataTracker.GetLightMapDimension() : InSettings.TargetLightMapResolution;
 				SrcModel.BuildSettings.SrcLightmapIndex = 0;
