World Partition - HLOD: Make sure that textures generated with the MeshApproximate algo are assigned to the TEXTUREGROUP_HierarchicalLOD LODGroup
* Now applying the LODGroup regardless of the HLOD builder - should prevent future builders from missing this LODGroup assignation

diff --git a/Engine/Plugins/Editor/WorldPartitionHLODUtilities/Source/Private/WorldPartition/HLOD/Builders/HLODBuilderMeshApproximate.cpp b/Engine/Plugins/Editor/WorldPartitionHLODUtilities/Source/Private/WorldPartition/HLOD/Builders/HLODBuilderMeshApproximate.cpp
--- a/Engine/Plugins/Editor/WorldPartitionHLODUtilities/Source/Private/WorldPartition/HLOD/Builders/HLODBuilderMeshApproximate.cpp
+++ b/Engine/Plugins/Editor/WorldPartitionHLODUtilities/Source/Private/WorldPartition/HLOD/Builders/HLODBuilderMeshApproximate.cpp
@@ -47,7 +47,7 @@ uint32 UHLODBuilderMeshApproximateSettings::GetCRC() const
 	FArchiveCrc32 Ar;
 
 	// Base key, changing this will force a rebuild of all HLODs from this builder
-	FString HLODBaseKey = "60D800B93F8B43789AB7FC0BF9F9BDC1";
+	FString HLODBaseKey = "4293829EF930455CA00E56A0B3F4E67F";
 	Ar << HLODBaseKey;
 
 	Ar << This.MeshApproximationSettings;
diff --git a/Engine/Plugins/Editor/WorldPartitionHLODUtilities/Source/Private/WorldPartition/HLOD/Utilities/WorldPartitionHLODUtilities.cpp b/Engine/Plugins/Editor/WorldPartitionHLODUtilities/Source/Private/WorldPartition/HLOD/Utilities/WorldPartitionHLODUtilities.cpp
--- a/Engine/Plugins/Editor/WorldPartitionHLODUtilities/Source/Private/WorldPartition/HLOD/Utilities/WorldPartitionHLODUtilities.cpp
+++ b/Engine/Plugins/Editor/WorldPartitionHLODUtilities/Source/Private/WorldPartition/HLOD/Utilities/WorldPartitionHLODUtilities.cpp
@@ -846,30 +846,62 @@ uint32 FWorldPartitionHLODUtilities::BuildHLOD(AWorldPartitionHLOD* InHLODActor)
 							{
 								StaticMeshComponent->bRayTracingFarField = true;
 							}
+						}
+					}
+				}
 
-							// Disable navigation data on HLODs
-							StaticMesh->MarkAsNotHavingNavigationData();
+				// Gather assets created during the HLOD generation process
+				struct FHLODAssets
+				{
+					TArray<UTexture*> Textures;
+					TArray<UStaticMesh*> StaticMeshes;
+				} HLODAssets;				
 
-							// Ensure we can perform line trace on HLOD static meshes. This is useful for all kind of editor features on HLODs (Actor placement, Play from here, Go Here, etc)
-							// Collision data will be stripped at cook.
-							StaticMesh->CreateBodySetup();
-							UBodySetup* BodySetup = StaticMesh->GetBodySetup();
-							BodySetup->DefaultInstance.SetCollisionEnabled(ECollisionEnabled::QueryOnly);
-							BodySetup->DefaultInstance.SetCollisionProfileName(UCollisionProfile::BlockAll_ProfileName);
-							for (int32 LODIndex = 0; LODIndex < StaticMesh->GetNumLODs(); ++LODIndex)
-							{
-								for (int32 SectionIndex = 0; SectionIndex < StaticMesh->GetNumSections(LODIndex); ++SectionIndex)
-								{
-									FMeshSectionInfo MeshSectionInfo = StaticMesh->GetSectionInfoMap().Get(LODIndex, SectionIndex);
-									MeshSectionInfo.bEnableCollision = true;
-									StaticMesh->GetSectionInfoMap().Set(LODIndex, SectionIndex, MeshSectionInfo);
-								}
-							}
+				// Retrieve assets first, as some of the calls we have to perform can't be done inside of a ForEachObjectWithOuter() (ex: Rename)
+				ForEachObjectWithOuter(HLODBuildContext.AssetsOuter, [&HLODAssets](UObject* Obj)
+				{
+					if (UStaticMesh* StaticMesh = Cast<UStaticMesh>(Obj))
+					{
+						HLODAssets.StaticMeshes.Add(StaticMesh);
+					}
 
-							// Rename owned static mesh
-							StaticMesh->Rename(*MakeUniqueObjectName(StaticMesh->GetOuter(), StaticMesh->GetClass(), *FString::Printf(TEXT("StaticMesh_%s"), *HLODLayer->GetName())).ToString());
+					if (UTexture* Texture = Cast<UTexture>(Obj))
+					{
+						HLODAssets.Textures.Add(Texture);
+					}
+				});
+
+				// Static meshes
+				for (UStaticMesh* StaticMesh : HLODAssets.StaticMeshes)
+				{
+					// Disable navigation data on HLODs
+					StaticMesh->MarkAsNotHavingNavigationData();
+
+					// Ensure we can perform line trace on HLOD static meshes. This is useful for all kind of editor features on HLODs (Actor placement, Play from here, Go Here, etc)
+					// Collision data will be stripped at cook.
+					StaticMesh->CreateBodySetup();
+					UBodySetup* BodySetup = StaticMesh->GetBodySetup();
+					BodySetup->DefaultInstance.SetCollisionEnabled(ECollisionEnabled::QueryOnly);
+					BodySetup->DefaultInstance.SetCollisionProfileName(UCollisionProfile::BlockAll_ProfileName);
+					for (int32 LODIndex = 0; LODIndex < StaticMesh->GetNumLODs(); ++LODIndex)
+					{
+						for (int32 SectionIndex = 0; SectionIndex < StaticMesh->GetNumSections(LODIndex); ++SectionIndex)
+						{
+							FMeshSectionInfo MeshSectionInfo = StaticMesh->GetSectionInfoMap().Get(LODIndex, SectionIndex);
+							MeshSectionInfo.bEnableCollision = true;
+							StaticMesh->GetSectionInfoMap().Set(LODIndex, SectionIndex, MeshSectionInfo);
 						}
 					}
+				
+					// Rename owned static mesh
+					StaticMesh->Rename(*MakeUniqueObjectName(StaticMesh->GetOuter(), StaticMesh->GetClass(), *FString::Printf(TEXT("StaticMesh_%s"), *HLODLayer->GetName())).ToString());
+				}
+
+				// Textures
+				for (UTexture* Texture : HLODAssets.Textures)
+				{
+					// Make sure all HLOD textures are assigned to the correct LODGroup
+					Texture->LODGroup = TEXTUREGROUP_HierarchicalLOD;
 				}
 			}
 
