[HLODBuilder] Apply UHLODBuilderInstancingSettings.bDisallowNanite even when meshes have only one LOD
DisallowNanite logic was originally added to fix meshes that have multiple LODs and set the mesh last LOD as the component ForcedLOD. We can extend the logic to all cases as that will allow meshes with single LOD to use Nanite fallback mesh for HLOD purposes.

diff --git a/Engine/Plugins/Editor/WorldPartitionHLODUtilities/Source/Private/WorldPartition/HLOD/Builders/HLODBuilderInstancing.cpp b/Engine/Plugins/Editor/WorldPartitionHLODUtilities/Source/Private/WorldPartition/HLOD/Builders/HLODBuilderInstancing.cpp
--- a/Engine/Plugins/Editor/WorldPartitionHLODUtilities/Source/Private/WorldPartition/HLOD/Builders/HLODBuilderInstancing.cpp
+++ b/Engine/Plugins/Editor/WorldPartitionHLODUtilities/Source/Private/WorldPartition/HLOD/Builders/HLODBuilderInstancing.cpp
@@ -47,7 +47,7 @@ TArray<UActorComponent*> UHLODBuilderInstancing::Build(const FHLODBuildContext&
 
 	TArray<UActorComponent*> HLODComponents = UHLODBuilder::BatchInstances(InSourceComponents);
 
-	// If requested, disallow Nanite on components whose mesh is Nanite enabled and has multiple LODs
+	// If requested, disallow Nanite on components whose mesh is Nanite enabled
 	if (InstancingSettings.bDisallowNanite)
 	{
 		for (UActorComponent* HLODComponent : HLODComponents)
@@ -56,7 +56,7 @@ TArray<UActorComponent*> UHLODBuilderInstancing::Build(const FHLODBuildContext&
 			{
 				if (UStaticMesh* StaticMesh = SMComponent->GetStaticMesh())
 				{
-					if (!SMComponent->bDisallowNanite && StaticMesh->HasValidNaniteData() && StaticMesh->GetNumLODs() > 1)
+					if (!SMComponent->bDisallowNanite && StaticMesh->HasValidNaniteData())
 					{
 						SMComponent->bDisallowNanite = true;
 						SMComponent->MarkRenderStateDirty();
