World Building - HLOD Generation: Changed some warnings to logs
* List input components when no HLODs are generated

diff --git a/Engine/Plugins/Editor/ProxyLODPlugin/Source/ProxyLOD/Private/ProxyLODPlugin.cpp b/Engine/Plugins/Editor/ProxyLODPlugin/Source/ProxyLOD/Private/ProxyLODPlugin.cpp
--- a/Engine/Plugins/Editor/ProxyLODPlugin/Source/ProxyLOD/Private/ProxyLODPlugin.cpp
+++ b/Engine/Plugins/Editor/ProxyLODPlugin/Source/ProxyLOD/Private/ProxyLODPlugin.cpp
@@ -677,7 +677,7 @@ void FVoxelizeMeshMerging::ProxyLOD(FMeshDescriptionArrayAdapter& InSrcGeometryA
 
 			if (!BBox.hasVolume())
 			{
-				UE_LOG(LogProxyLODMeshReduction, Warning, TEXT("Empty bounding box for all static meshes input, no output mesh will be generated."));
+				UE_LOG(LogProxyLODMeshReduction, Log, TEXT("Empty bounding box for all static meshes input, no output mesh will be generated."));
 
 				FailedDelegate.ExecuteIfBound(InJobGUID, TEXT("ProxyLOD invalid input geometry"));
 
diff --git a/Engine/Plugins/Editor/WorldPartitionHLODUtilities/Source/Private/WorldPartition/HLOD/Utilities/WorldPartitionHLODUtilities.cpp b/Engine/Plugins/Editor/WorldPartitionHLODUtilities/Source/Private/WorldPartition/HLOD/Utilities/WorldPartitionHLODUtilities.cpp
--- a/Engine/Plugins/Editor/WorldPartitionHLODUtilities/Source/Private/WorldPartition/HLOD/Utilities/WorldPartitionHLODUtilities.cpp
+++ b/Engine/Plugins/Editor/WorldPartitionHLODUtilities/Source/Private/WorldPartition/HLOD/Utilities/WorldPartitionHLODUtilities.cpp
@@ -804,7 +804,11 @@ uint32 FWorldPartitionHLODUtilities::BuildHLOD(AWorldPartitionHLOD* InHLODActor)
 
 			if (BuildResult.HLODComponents.IsEmpty())
 			{
-				UE_LOG(LogHLODBuilder, Log, TEXT("HLOD generation created no component for %s"), *InHLODActor->GetActorLabel());
+				UE_LOG(LogHLODBuilder, Log, TEXT("HLOD generation created no component for HLOD actor %s. Listing source components:"), *InHLODActor->GetActorLabel());
+				for (UActorComponent* SourceComponent : HLODBuildContext.SourceComponents)
+				{
+					UE_LOG(LogHLODBuilder, Log, TEXT("\t* %s"), *SourceComponent->GetPathName());
+				}
 			}
 
 			// Ideally, this should be performed elsewhere, to allow more flexibility in the HLOD generation
diff --git a/Engine/Plugins/Experimental/MeshModelingToolsetExp/Source/GeometryProcessingAdapters/Private/ApproximateActorsImpl.cpp b/Engine/Plugins/Experimental/MeshModelingToolsetExp/Source/GeometryProcessingAdapters/Private/ApproximateActorsImpl.cpp
--- a/Engine/Plugins/Experimental/MeshModelingToolsetExp/Source/GeometryProcessingAdapters/Private/ApproximateActorsImpl.cpp
+++ b/Engine/Plugins/Experimental/MeshModelingToolsetExp/Source/GeometryProcessingAdapters/Private/ApproximateActorsImpl.cpp
@@ -434,7 +434,7 @@ static TSharedPtr<FApproximationMeshData> GenerateApproximationMesh(
 	// avoid insane memory usage
 	if (VoxelDimTarget > Options.ClampVoxelDimension)
 	{
-		UE_LOG(LogApproximateActors, Warning, TEXT("very large voxel size %d clamped to %d"), VoxelDimTarget, Options.ClampVoxelDimension);
+		UE_LOG(LogApproximateActors, Log, TEXT("very large voxel size %d clamped to %d"), VoxelDimTarget, Options.ClampVoxelDimension);
 		VoxelDimTarget = Options.ClampVoxelDimension;
 	}
 
