WorldPartition - HLOD: Bump HLOD key for HLODs built with MeshSimplify
* Will ensure they are rebuilt following the fix made in CL 42931841

diff --git a/Engine/Plugins/Editor/WorldPartitionHLODUtilities/Source/Private/WorldPartition/HLOD/Builders/HLODBuilderMeshSimplify.cpp b/Engine/Plugins/Editor/WorldPartitionHLODUtilities/Source/Private/WorldPartition/HLOD/Builders/HLODBuilderMeshSimplify.cpp
--- a/Engine/Plugins/Editor/WorldPartitionHLODUtilities/Source/Private/WorldPartition/HLOD/Builders/HLODBuilderMeshSimplify.cpp
+++ b/Engine/Plugins/Editor/WorldPartitionHLODUtilities/Source/Private/WorldPartition/HLOD/Builders/HLODBuilderMeshSimplify.cpp
@@ -44,7 +44,7 @@ uint32 UHLODBuilderMeshSimplifySettings::GetCRC() const
 	FArchiveCrc32 Ar;
 
 	// Base mesh simplify key, changing this will force a rebuild of all HLODs from this builder
-	FString HLODBaseKey = "9015D801CD0E420B9D2B0CB997A97813";
+	FString HLODBaseKey = "AD4D6BE5CFA5447BB43DD31171A429BF";
 	Ar << HLODBaseKey;
 
 	Ar << This.MeshSimplifySettings;
