World Partition: always initialize FHLODCreationParams::TargetWorld.

diff --git a/Engine/Plugins/Editor/WorldPartitionHLODUtilities/Source/Private/WorldPartition/HLOD/Utilities/WorldPartitionHLODUtilities.cpp b/Engine/Plugins/Editor/WorldPartitionHLODUtilities/Source/Private/WorldPartition/HLOD/Utilities/WorldPartitionHLODUtilities.cpp
--- a/Engine/Plugins/Editor/WorldPartitionHLODUtilities/Source/Private/WorldPartition/HLOD/Utilities/WorldPartitionHLODUtilities.cpp
+++ b/Engine/Plugins/Editor/WorldPartitionHLODUtilities/Source/Private/WorldPartition/HLOD/Utilities/WorldPartitionHLODUtilities.cpp
@@ -208,14 +208,7 @@ TArray<AWorldPartitionHLOD*> FWorldPartitionHLODUtilities::CreateHLODActors(FHLO
 			FActorSpawnParameters SpawnParams;
 			SpawnParams.Name = HLODActorName;
 			SpawnParams.NameMode = FActorSpawnParameters::ESpawnActorNameMode::Required_Fatal;
-			if (InCreationParams.bIsStandalone)
-			{
-				HLODActor = InCreationParams.TargetWorld->SpawnActor<AWorldPartitionHLOD>(HLODLayer->GetHLODActorClass(), SpawnParams);
-			}
-			else
-			{
-				HLODActor = InCreationParams.WorldPartition->GetWorld()->SpawnActor<AWorldPartitionHLOD>(HLODLayer->GetHLODActorClass(), SpawnParams);
-			}
+			HLODActor = InCreationParams.TargetWorld->SpawnActor<AWorldPartitionHLOD>(HLODLayer->GetHLODActorClass(), SpawnParams);
 
 			check(HLODActor->GetContentBundleGuid() == InCreationParams.ContentBundleGuid);
 			check(HLODActor->GetExternalDataLayerAsset() == InCreationParams.GetExternalDataLayerAsset());
diff --git a/Engine/Source/Runtime/Engine/Private/WorldPartition/RuntimeSpatialHash/RuntimeSpatialHashHLOD.cpp b/Engine/Source/Runtime/Engine/Private/WorldPartition/RuntimeSpatialHash/RuntimeSpatialHashHLOD.cpp
--- a/Engine/Source/Runtime/Engine/Private/WorldPartition/RuntimeSpatialHash/RuntimeSpatialHashHLOD.cpp
+++ b/Engine/Source/Runtime/Engine/Private/WorldPartition/RuntimeSpatialHash/RuntimeSpatialHashHLOD.cpp
@@ -211,6 +211,7 @@ static TArray<FGuid> GenerateHLODActorsForGrid(UWorldPartition* WorldPartition,
 
 				FHLODCreationParams CreationParams;
 				CreationParams.WorldPartition = WorldPartition;
+				CreationParams.TargetWorld = WorldPartition->GetWorld();
 				CreationParams.CellGuid = CellGuid;
 				CreationParams.CellName = CellName;
 				CreationParams.CellBounds = CellBounds;
