Fix issue where texture LODGroup was possibly changed outside of a PreEditChange() / PostEditChange() block

diff --git a/Engine/Plugins/Editor/WorldPartitionHLODUtilities/Source/Private/WorldPartition/HLOD/Utilities/WorldPartitionHLODUtilities.cpp b/Engine/Plugins/Editor/WorldPartitionHLODUtilities/Source/Private/WorldPartition/HLOD/Utilities/WorldPartitionHLODUtilities.cpp
--- a/Engine/Plugins/Editor/WorldPartitionHLODUtilities/Source/Private/WorldPartition/HLOD/Utilities/WorldPartitionHLODUtilities.cpp
+++ b/Engine/Plugins/Editor/WorldPartitionHLODUtilities/Source/Private/WorldPartition/HLOD/Utilities/WorldPartitionHLODUtilities.cpp
@@ -919,7 +919,12 @@ uint32 FWorldPartitionHLODUtilities::BuildHLOD(AWorldPartitionHLOD* InHLODActor)
 				for (UTexture* Texture : HLODAssets.Textures)
 				{
 					// Make sure all HLOD textures are assigned to the correct LODGroup
-					Texture->LODGroup = TEXTUREGROUP_HierarchicalLOD;
+					if (Texture->LODGroup != TEXTUREGROUP_HierarchicalLOD)
+					{
+						Texture->PreEditChange(nullptr);
+						Texture->LODGroup = TEXTUREGROUP_HierarchicalLOD;
+						Texture->PostEditChange();
+					}
 				}
 			}
 
