World Partition - HLOD: When applying the destructible HLOD modifier during HLOD builds, skip processing of invalid component data

diff --git a/Engine/Plugins/Editor/WorldPartitionHLODUtilities/Source/Private/WorldPartition/HLOD/Modifiers/HLODModifierMeshDestruction.cpp b/Engine/Plugins/Editor/WorldPartitionHLODUtilities/Source/Private/WorldPartition/HLOD/Modifiers/HLODModifierMeshDestruction.cpp
--- a/Engine/Plugins/Editor/WorldPartitionHLODUtilities/Source/Private/WorldPartition/HLOD/Modifiers/HLODModifierMeshDestruction.cpp
+++ b/Engine/Plugins/Editor/WorldPartitionHLODUtilities/Source/Private/WorldPartition/HLOD/Modifiers/HLODModifierMeshDestruction.cpp
@@ -54,26 +54,6 @@ public:
 	{
 		check(DestructibleActors.IsEmpty());
 
-		// Create component to actor map & compute components to wedges offsets.
-		TMap<uint32, uint32> ComponentToActorIndex;
-		TArray<uint32> ComponentToWedgeOffsets;
-		ComponentToWedgeOffsets.SetNumZeroed(MergedComponents.Num());
-
-		for (int32 ComponentIndex = 0; ComponentIndex < MergedComponents.Num(); ++ComponentIndex)
-		{
-			ComponentToWedgeOffsets[ComponentIndex] = DataTracker.GetComponentToWedgeMappng(ComponentIndex, /*LODIndex=*/0);
-
-			UStaticMeshComponent* StaticMeshComponent = MergedComponents[ComponentIndex];
-			if (AActor* Actor = StaticMeshComponent->GetOwner())
-			{
-				if (Actor->Implements<UWorldPartitionDestructibleInHLODInterface>())
-				{
-					uint32 ActorIndex = DestructibleActors.AddUnique(Actor);
-					ComponentToActorIndex.Add(ComponentIndex, ActorIndex);
-				}
-			}
-		}
-
 		// Clear vertex colors
 		FMeshDescription& MergedMesh = MergedMeshLODs[0];
 		TVertexInstanceAttributesRef<FVector4f> VertexInstanceColors = MergedMesh.VertexInstanceAttributes().GetAttributesRef<FVector4f>(MeshAttribute::VertexInstance::Color);
@@ -82,23 +62,44 @@ public:
 			VertexInstanceColors[VertexInstanceID] = FVector4f(0.0f, 0.0f, 0.0f, 0.0f);
 		}
 
-		// Store actor index + 1 in the wedge colors for each part of the merged mesh
+		const uint32 MergedMeshVertexInstancesNum = MergedMesh.VertexInstances().Num();
+
 		for (int32 ComponentIndex = 0; ComponentIndex < MergedComponents.Num(); ++ComponentIndex)
 		{
-			const uint32 ComponentStart = ComponentToWedgeOffsets[ComponentIndex];
-			const uint32 ComponentEnd = ComponentToWedgeOffsets.IsValidIndex(ComponentIndex + 1) ? ComponentToWedgeOffsets[ComponentIndex + 1] : MergedMesh.VertexInstances().Num();
-			FColor ActorColor = FColor::Black;
+			// Start index
+			const uint32 ComponentVertexColorIndexStart = DataTracker.GetComponentToWedgeMappng(ComponentIndex, /*LODIndex=*/0);
+			if (ComponentVertexColorIndexStart == INDEX_NONE)
+			{
+				continue;
+			}
 
-			const uint32* ActorIndex = ComponentToActorIndex.Find(ComponentIndex);
-			if (ActorIndex)
+			// End index
+			uint32 ComponentVertexColorIndexEnd = MergedMeshVertexInstancesNum;
+			for (int32 NextComponentIndex = ComponentIndex + 1; NextComponentIndex < MergedComponents.Num(); ++NextComponentIndex)
 			{
-				ActorColor.DWColor() = *ActorIndex + 1;
+				// Skip invalid entries
+				uint32 NextComponentVertexColorIndexStart = DataTracker.GetComponentToWedgeMappng(NextComponentIndex, /*LODIndex=*/0);
+				if (NextComponentVertexColorIndexStart != INDEX_NONE)
+				{
+					ComponentVertexColorIndexEnd = NextComponentVertexColorIndexStart;
+					break;
+				}
+			}
+
+			FColor ActorColor = FColor::Black;
+			UStaticMeshComponent* StaticMeshComponent = MergedComponents[ComponentIndex];
+			if (AActor* Actor = StaticMeshComponent->GetOwner())
+			{
+				if (Actor->Implements<UWorldPartitionDestructibleInHLODInterface>())
+				{
+					uint32 ActorIndex = DestructibleActors.AddUnique(Actor);
+					ActorColor.DWColor() = ActorIndex + 1;
+				}
 			}
 
-			for (uint32 Index = ComponentStart; Index < ComponentEnd; ++Index)
+			for (uint32 ComponentVertexColorIndex = ComponentVertexColorIndexStart; ComponentVertexColorIndex < ComponentVertexColorIndexEnd; ++ComponentVertexColorIndex)
 			{
-				FVertexInstanceID VertexInstanceID(Index);
-				VertexInstanceColors[FVertexInstanceID(Index)] = FLinearColor(ActorColor);
+				VertexInstanceColors[FVertexInstanceID(ComponentVertexColorIndex)] = FLinearColor(ActorColor);
 			}
 		}
 	}
