World Partition - HLOD: Use a 64-bit hash value to generate the HLOD actor object name
* A 32-bit hash is often sufficient, but could lead to collisions when dealing with large worlds with many HLOD actors
* Retain the option to use a 32-bit hash, as the change to a new hash will force the creation of a whole set of new HLOD actors.

diff --git a/Engine/Plugins/Editor/WorldPartitionHLODUtilities/Source/Private/WorldPartition/HLOD/Utilities/WorldPartitionHLODUtilities.cpp b/Engine/Plugins/Editor/WorldPartitionHLODUtilities/Source/Private/WorldPartition/HLOD/Utilities/WorldPartitionHLODUtilities.cpp
--- a/Engine/Plugins/Editor/WorldPartitionHLODUtilities/Source/Private/WorldPartition/HLOD/Utilities/WorldPartitionHLODUtilities.cpp
+++ b/Engine/Plugins/Editor/WorldPartitionHLODUtilities/Source/Private/WorldPartition/HLOD/Utilities/WorldPartitionHLODUtilities.cpp
@@ -43,6 +43,12 @@
 #include "UObject/GCObjectScopeGuard.h"
 #include "ContentStreaming.h"
 
+TAutoConsoleVariable<bool> CVarHLODUseLegacy32BitNameHash(
+	TEXT("wp.Editor.HLOD.UseLegacy32BitNameHash"),
+	false,
+	TEXT("Force the use of 32-bit hashing to use when generating HLOD actor names. Switching to 64-bit will result in new names/new actor files."),
+	ECVF_ReadOnly);
+
 static uint32 ComputeHLODHash(AWorldPartitionHLOD* InHLODActor, const TArray<UActorComponent*>& InSourceComponents)
 {
 	FArchiveCrc32 Ar;
@@ -142,6 +148,35 @@ void AddSourceActor(const FStreamingGenerationActorDescView& ActorDescView, cons
 	}
 }
 
+uint64 ComputeHLODActorUniqueHash(const UHLODLayer* HLODLayer, const FGuid CellGuid)
+{
+	uint64 HLODActorHash = 0;
+	if (!CVarHLODUseLegacy32BitNameHash.GetValueOnGameThread())
+	{
+		const FString HLODLayerName = HLODLayer->GetName();
+		HLODActorHash = CityHash64WithSeed(reinterpret_cast<const char*>(HLODLayerName.GetCharArray().GetData()), HLODLayerName.GetCharArray().Num(), HLODActorHash);
+		HLODActorHash = CityHash64WithSeed((char*)&CellGuid, sizeof(FGuid), HLODActorHash);
+		if (HLODLayer->GetHLODActorClass() != AWorldPartitionHLOD::StaticClass())
+		{
+			const FString HLODActorClassPathName = HLODLayer->GetHLODActorClass()->GetPathName();
+			HLODActorHash = CityHash64WithSeed(reinterpret_cast<const char*>(HLODActorClassPathName.GetCharArray().GetData()), HLODActorClassPathName.GetCharArray().Num(), HLODActorHash);
+		}
+	}
+	else
+	{
+		const uint32 HLODLayerNameHash = FCrc::StrCrc32(*HLODLayer->GetName());
+		const uint32 CellGuidHash = GetTypeHash(CellGuid);
+		uint32 HLODActorHash32 = HashCombineFast(HLODLayerNameHash, CellGuidHash);
+		if (HLODLayer->GetHLODActorClass() != AWorldPartitionHLOD::StaticClass())
+		{
+			const uint32 HLODActorClassHash = FCrc::StrCrc32(*HLODLayer->GetHLODActorClass()->GetPathName());
+			HLODActorHash32 = HashCombineFast(HLODActorHash32, HLODActorClassHash);
+		}
+		HLODActorHash = HLODActorHash32;
+	}
+	return HLODActorHash;			
+}
+
 TArray<AWorldPartitionHLOD*> FWorldPartitionHLODUtilities::CreateHLODActors(FHLODCreationContext& InCreationContext, const FHLODCreationParams& InCreationParams, const TArray<IStreamingGenerationContext::FActorInstance>& InActors)
 {
 	TMap<UHLODLayer*, TMap<FGuid, FWorldPartitionRuntimeCellObjectMapping>> SourceActorsPerHLODLayer;
@@ -175,23 +210,8 @@ TArray<AWorldPartitionHLOD*> FWorldPartitionHLODUtilities::CreateHLODActors(FHLO
 		TMap<FGuid, FWorldPartitionRuntimeCellObjectMapping>& SourceActorsMap = Pair.Value;
 		check(!SourceActorsMap.IsEmpty());
 
-		auto ComputeHLODActorUniqueHash = [](const UHLODLayer* HLODLayer, const FGuid CellGuid)
-		{
-			const uint32 HLODLayerNameHash = FCrc::StrCrc32(*HLODLayer->GetName());
-			const uint32 CellGuidHash = GetTypeHash(CellGuid);
-			uint32 HLODActorHash = HashCombineFast(HLODLayerNameHash, CellGuidHash);
-
-			if (HLODLayer->GetHLODActorClass() != AWorldPartitionHLOD::StaticClass())
-			{
-				const uint32 HLODActorClassHash = FCrc::StrCrc32(*HLODLayer->GetHLODActorClass()->GetPathName());
-				HLODActorHash = HashCombineFast(HLODActorHash, HLODActorClassHash);
-			}
-
-			return HLODActorHash;
-		};
-
 		uint64 HLODActorHash = ComputeHLODActorUniqueHash(HLODLayer, InCreationParams.CellGuid);
-		FName HLODActorName = *FString::Printf(TEXT("%s_%016llx"), *HLODLayer->GetName(), HLODActorHash);
+		FName HLODActorName = *FString::Printf(TEXT("%s_%016" UINT64_x_FMT), *HLODLayer->GetName(), HLODActorHash);
 
 		AWorldPartitionHLOD* HLODActor = nullptr;
 		FWorldPartitionHandle HLODActorHandle;
