Use a TUniquePtr to replace a manually handled allocation

diff --git a/Engine/Plugins/Editor/WorldPartitionHLODUtilities/Source/Private/WorldPartition/HLOD/Modifiers/HLODModifierMeshDestruction.cpp b/Engine/Plugins/Editor/WorldPartitionHLODUtilities/Source/Private/WorldPartition/HLOD/Modifiers/HLODModifierMeshDestruction.cpp
--- a/Engine/Plugins/Editor/WorldPartitionHLODUtilities/Source/Private/WorldPartition/HLOD/Modifiers/HLODModifierMeshDestruction.cpp
+++ b/Engine/Plugins/Editor/WorldPartitionHLODUtilities/Source/Private/WorldPartition/HLOD/Modifiers/HLODModifierMeshDestruction.cpp
@@ -145,7 +145,7 @@ void UWorldPartitionHLODModifierMeshDestruction::BeginHLODBuild(const FHLODBuild
 {
 	CachedWorldPosition = InHLODBuildContext.WorldPosition;
 
-	DestructionMeshMergeExtension = new FDestructionMeshMergeExtension();
+	DestructionMeshMergeExtension = MakeUnique<FDestructionMeshMergeExtension>();
 }
 
 void UWorldPartitionHLODModifierMeshDestruction::EndHLODBuild(TArray<UActorComponent*>& InOutComponents)
@@ -285,5 +285,5 @@ void UWorldPartitionHLODModifierMeshDestruction::EndHLODBuild(TArray<UActorCompo
 		InOutComponents.Add(DestructibleHLODComponent);
 	}
 
-	delete DestructionMeshMergeExtension;
+	DestructionMeshMergeExtension.Reset();
 }
diff --git a/Engine/Plugins/Editor/WorldPartitionHLODUtilities/Source/Public/WorldPartition/HLOD/Modifiers/HLODModifierMeshDestruction.h b/Engine/Plugins/Editor/WorldPartitionHLODUtilities/Source/Public/WorldPartition/HLOD/Modifiers/HLODModifierMeshDestruction.h
--- a/Engine/Plugins/Editor/WorldPartitionHLODUtilities/Source/Public/WorldPartition/HLOD/Modifiers/HLODModifierMeshDestruction.h
+++ b/Engine/Plugins/Editor/WorldPartitionHLODUtilities/Source/Public/WorldPartition/HLOD/Modifiers/HLODModifierMeshDestruction.h
@@ -3,6 +3,7 @@
 #pragma once
 
 #include "WorldPartition/HLOD/HLODModifier.h"
+#include "Templates/UniquePtr.h"
 #include "HLODModifierMeshDestruction.generated.h"
 
 #define UE_API WORLDPARTITIONHLODUTILITIES_API
@@ -22,7 +23,7 @@ class UWorldPartitionHLODModifierMeshDestruction : public UWorldPartitionHLODMod
 
 private:
 	// Gathers a list of actors included in the merged mesh
-	FDestructionMeshMergeExtension* DestructionMeshMergeExtension;
+	TUniquePtr<FDestructionMeshMergeExtension> DestructionMeshMergeExtension;
 
 	// Cached from the HLOD Build Context provided to BeginHLODBuild()
 	FVector CachedWorldPosition;
