Fixed issue with incorrect buffer length provided to CityHash64 when dealing with FString and TCHARs

diff --git a/Engine/Plugins/Editor/WorldPartitionHLODUtilities/Source/Private/WorldPartition/HLOD/Utilities/WorldPartitionHLODUtilities.cpp b/Engine/Plugins/Editor/WorldPartitionHLODUtilities/Source/Private/WorldPartition/HLOD/Utilities/WorldPartitionHLODUtilities.cpp
--- a/Engine/Plugins/Editor/WorldPartitionHLODUtilities/Source/Private/WorldPartition/HLOD/Utilities/WorldPartitionHLODUtilities.cpp
+++ b/Engine/Plugins/Editor/WorldPartitionHLODUtilities/Source/Private/WorldPartition/HLOD/Utilities/WorldPartitionHLODUtilities.cpp
@@ -154,12 +154,12 @@ uint64 ComputeHLODActorUniqueHash(const UHLODLayer* HLODLayer, const FGuid CellG
 	if (!CVarHLODUseLegacy32BitNameHash.GetValueOnGameThread())
 	{
 		const FString HLODLayerName = HLODLayer->GetName();
-		HLODActorHash = CityHash64WithSeed(reinterpret_cast<const char*>(HLODLayerName.GetCharArray().GetData()), HLODLayerName.GetCharArray().Num(), HLODActorHash);
+		HLODActorHash = CityHash64WithSeed(reinterpret_cast<const char*>(HLODLayerName.GetCharArray().GetData()), HLODLayerName.GetCharArray().Num() * sizeof(FString::ElementType), HLODActorHash);
 		HLODActorHash = CityHash64WithSeed((char*)&CellGuid, sizeof(FGuid), HLODActorHash);
 		if (HLODLayer->GetHLODActorClass() != AWorldPartitionHLOD::StaticClass())
 		{
 			const FString HLODActorClassPathName = HLODLayer->GetHLODActorClass()->GetPathName();
-			HLODActorHash = CityHash64WithSeed(reinterpret_cast<const char*>(HLODActorClassPathName.GetCharArray().GetData()), HLODActorClassPathName.GetCharArray().Num(), HLODActorHash);
+			HLODActorHash = CityHash64WithSeed(reinterpret_cast<const char*>(HLODActorClassPathName.GetCharArray().GetData()), HLODActorClassPathName.GetCharArray().Num() * sizeof(FString::ElementType), HLODActorHash);
 		}
 	}
 	else
