World Partition - HLOD: Fixed destructible HLOD logic issue when no destructible actors were found in the source meshes
* Folding logic on an empty ISMActorMapping array was creating a bogus entry and causing a crash
* Also improved the amount of components possibly added to PackedISMData.ISMCs, some might not have been needed

diff --git a/Engine/Plugins/Editor/WorldPartitionHLODUtilities/Source/Private/WorldPartition/HLOD/Modifiers/HLODModifierMeshDestruction.cpp b/Engine/Plugins/Editor/WorldPartitionHLODUtilities/Source/Private/WorldPartition/HLOD/Modifiers/HLODModifierMeshDestruction.cpp
--- a/Engine/Plugins/Editor/WorldPartitionHLODUtilities/Source/Private/WorldPartition/HLOD/Modifiers/HLODModifierMeshDestruction.cpp
+++ b/Engine/Plugins/Editor/WorldPartitionHLODUtilities/Source/Private/WorldPartition/HLOD/Modifiers/HLODModifierMeshDestruction.cpp
@@ -168,18 +168,25 @@ void UWorldPartitionHLODModifierMeshDestruction::EndHLODBuild(TArray<UActorCompo
 		};
 		TArray<FISMActorMapping> ISMActorMapping;
 
+		TArray<UHLODInstancedStaticMeshComponent*> DestructibleHLODISMCs;
+
 		// Build the full ISM to actor mapping by retrieving the mapping data from the UHLODInstancedStaticMeshComponent components
-		for (uint32 ComponentIndex = 0; ComponentIndex < (uint32)HLODISMCs.Num(); ComponentIndex++)
+		for (UHLODInstancedStaticMeshComponent* Component : HLODISMCs)
 		{
-			UHLODInstancedStaticMeshComponent* Component = HLODISMCs[ComponentIndex];
 			for (const FISMComponentBatcher::FComponentToInstancesMapping& ComponentToInstancesMapping : Component->GetSourceComponentsToInstancesMap())
 			{
 				const UActorComponent* SourceComponent = ComponentToInstancesMapping.Component;
 				const AActor* SourceActor = SourceComponent->GetOwner();
 				if (SourceActor->Implements<UWorldPartitionDestructibleInHLODInterface>())
 				{
+					// Add component to the destructible HLOD ISMC array if it wasn't already added
+					if (DestructibleHLODISMCs.IsEmpty() || DestructibleHLODISMCs.Last() != Component)
+					{
+						DestructibleHLODISMCs.Add(Component);
+					}
+
 					uint32 ActorIndex = DestructibleActors.AddUnique(SourceActor->GetFName());
-					ISMActorMapping.Emplace(ActorIndex, ComponentIndex, ComponentToInstancesMapping.InstancesStart, ComponentToInstancesMapping.InstancesCount);
+					ISMActorMapping.Emplace(ActorIndex, DestructibleHLODISMCs.Num() - 1, ComponentToInstancesMapping.InstancesStart, ComponentToInstancesMapping.InstancesCount);
 				}
 			}
 		}
@@ -202,6 +209,7 @@ void UWorldPartitionHLODModifierMeshDestruction::EndHLODBuild(TArray<UActorCompo
 		});
 
 		// Fold entries with the same ActorIndex and ComponentIndex, and adjust the InstancingStart & InstancingCount
+		if (!ISMActorMapping.IsEmpty())
 		{
 			int32 Write = 0; // index of the current output slot
 
@@ -239,7 +247,7 @@ void UWorldPartitionHLODModifierMeshDestruction::EndHLODBuild(TArray<UActorCompo
 			ISMActorMapping.SetNum(Write + 1, EAllowShrinking::No);
 		}
 
-		PackedISMData.ISMCs = HLODISMCs;
+		PackedISMData.ISMCs = DestructibleHLODISMCs;
 
 		// Build the optimal representation of those mapping into our FHLODInstancingPackedMappingData structure
 		int32 CurrentIndex = 0;
diff --git a/Engine/Source/Runtime/Engine/Public/WorldPartition/HLOD/DestructibleHLODComponent.h b/Engine/Source/Runtime/Engine/Public/WorldPartition/HLOD/DestructibleHLODComponent.h
--- a/Engine/Source/Runtime/Engine/Public/WorldPartition/HLOD/DestructibleHLODComponent.h
+++ b/Engine/Source/Runtime/Engine/Public/WorldPartition/HLOD/DestructibleHLODComponent.h
@@ -109,8 +109,8 @@ struct FComponentInstanceMapping
 
 	static FComponentInstanceMapping Make(uint32 ComponentIndex, uint32 InstanceStart, uint32 InstanceCount)
 	{
-		checkf((ComponentIndex & ~kComponentIndexMask) == 0, TEXT("ComponentIdx exceeds %u bits"), kComponentIndexBits);
-		checkf((InstanceCount & ~kItemCountMask) == 0, TEXT("FItem::Make: ItemCount exceeds %u bits"), kItemCountBits);
+		checkf((ComponentIndex & ~kComponentIndexMask) == 0, TEXT("ComponentIndex exceeds %u bits"), kComponentIndexBits);
+		checkf((InstanceCount & ~kItemCountMask) == 0, TEXT("InstanceCount exceeds %u bits"), kItemCountBits);
 		FComponentInstanceMapping Mapping;
 		Mapping.A = ((ComponentIndex & kComponentIndexMask) << kComponentIndexShift) | (InstanceCount & kItemCountMask);
 		Mapping.B = InstanceStart; // full 32-bit Start
@@ -170,7 +170,7 @@ struct FActorInstanceMappingsRef
 
 	static FActorInstanceMappingsRef MakeMappingInline(uint32 ComponentIndex, uint32 InstanceStart, uint32 InstanceCount)
 	{
-		checkf((ComponentIndex & ~kComponentIndexMask) == 0, TEXT("ComponentIndex exceeds %u bits"), kComponentIndexShift);
+		checkf((ComponentIndex & ~kComponentIndexMask) == 0, TEXT("ComponentIndex exceeds %u bits"), kComponentIndexBits);
 		checkf((InstanceCount & ~kItemCountMask) == 0, TEXT("InstanceCount exceeds %u bits"), kItemCountBits);
 		FActorInstanceMappingsRef Mapping;
 		Mapping.A = kInlineTagMask
