World Partition ? HLOD: Fix determinism issue related to computed HLOD bounds
* HLOD bounds for the LHGrid partition were previously derived from the full source cell bounds, which included non-HLOD relevant actors.
* As a result, bounds changes on non-HLOD relevant actors could cause unnecessary resaves of HLOD actors during setup.
* HLOD bounds are now computed at build time from the actual HLOD-relevant components only.

diff --git a/Engine/Plugins/Editor/WorldPartitionHLODUtilities/Source/Private/WorldPartition/HLOD/Utilities/WorldPartitionHLODUtilities.cpp b/Engine/Plugins/Editor/WorldPartitionHLODUtilities/Source/Private/WorldPartition/HLOD/Utilities/WorldPartitionHLODUtilities.cpp
--- a/Engine/Plugins/Editor/WorldPartitionHLODUtilities/Source/Private/WorldPartition/HLOD/Utilities/WorldPartitionHLODUtilities.cpp
+++ b/Engine/Plugins/Editor/WorldPartitionHLODUtilities/Source/Private/WorldPartition/HLOD/Utilities/WorldPartitionHLODUtilities.cpp
@@ -78,7 +78,13 @@ static void ComputeHLODHash(FHLODHashBuilder& InHashBuilder, AWorldPartitionHLOD
 	UHLODBuilder::ComputeHLODHash(InHashBuilder, InSourceComponents);
 }
 
-void AddSourceActor(const FStreamingGenerationActorDescView& ActorDescView, const FName WorldPackageName, TMap<FGuid, FWorldPartitionRuntimeCellObjectMapping>& SourceActors)
+struct FPerHLODLayerData
+{
+	TMap<FGuid, FWorldPartitionRuntimeCellObjectMapping> SourceActors;
+	FBox Bounds = FBox(ForceInit);
+};
+
+void AddSourceActor(const FStreamingGenerationActorDescView& ActorDescView, const FName WorldPackageName, FPerHLODLayerData& PerHLODLayerData)
 {
 	const FName ActorPath = *ActorDescView.GetActorSoftPath().ToString();
 
@@ -89,8 +95,10 @@ void AddSourceActor(const FStreamingGenerationActorDescView& ActorDescView, cons
 	const FTransform ContainerTransform = ContainerInstance->GetTransform();
 	const FName ContainerPackage = ContainerInstance->GetContainerPackage();
 
+	PerHLODLayerData.Bounds += ActorDescView.GetRuntimeBounds();
+
 	FGuid ActorInstanceGuid = ContainerID.GetActorGuid(ActorDescView.GetGuid());	
-	FWorldPartitionRuntimeCellObjectMapping& ActorMapping = SourceActors.FindOrAdd(ActorInstanceGuid);
+	FWorldPartitionRuntimeCellObjectMapping& ActorMapping = PerHLODLayerData.SourceActors.FindOrAdd(ActorInstanceGuid);
 	if (!ActorMapping.ActorInstanceGuid.IsValid())
 	{
 		ActorMapping = FWorldPartitionRuntimeCellObjectMapping(
@@ -111,7 +119,7 @@ void AddSourceActor(const FStreamingGenerationActorDescView& ActorDescView, cons
 		for (const FGuid& ReferenceGuid : ActorDescView.GetReferences())
 		{
 			const FStreamingGenerationActorDescView& RefActorDescView = ActorDescViewMap.FindByGuidChecked(ReferenceGuid);
-			AddSourceActor(RefActorDescView, WorldPackageName, SourceActors);
+			AddSourceActor(RefActorDescView, WorldPackageName, PerHLODLayerData);
 		}
 
 		// Add its editor references
@@ -119,7 +127,7 @@ void AddSourceActor(const FStreamingGenerationActorDescView& ActorDescView, cons
 		{
 			const FWorldPartitionActorDescInstance& ReferenceActorDesc = ContainerInstance->GetActorDescInstanceChecked(EditorReferenceGuid);
 			FGuid EditorRefInstanceGuid = ContainerID.GetActorGuid(EditorReferenceGuid);
-			FWorldPartitionRuntimeCellObjectMapping& EditorRefMapping = SourceActors.FindOrAdd(EditorRefInstanceGuid);
+			FWorldPartitionRuntimeCellObjectMapping& EditorRefMapping = PerHLODLayerData.SourceActors.FindOrAdd(EditorRefInstanceGuid);
 			if (!EditorRefMapping.ActorInstanceGuid.IsValid())
 			{
 				EditorRefMapping = FWorldPartitionRuntimeCellObjectMapping(
@@ -171,7 +179,7 @@ uint64 ComputeHLODActorUniqueHash(const UHLODLayer* HLODLayer, const FGuid CellG
 
 TArray<AWorldPartitionHLOD*> FWorldPartitionHLODUtilities::CreateHLODActors(FHLODCreationContext& InCreationContext, const FHLODCreationParams& InCreationParams, const TArray<IStreamingGenerationContext::FActorInstance>& InActors)
 {
-	TMap<UHLODLayer*, TMap<FGuid, FWorldPartitionRuntimeCellObjectMapping>> SourceActorsPerHLODLayer;
+	TMap<UHLODLayer*, FPerHLODLayerData> PerHLODLayerDataMap;
 
 	FName WorldPackageName = InCreationParams.WorldPartition->GetWorld()->GetPackage()->GetFName();
 
@@ -203,17 +211,18 @@ TArray<AWorldPartitionHLOD*> FWorldPartitionHLODUtilities::CreateHLODActors(FHLO
 					}
 				}
 
-				TMap<FGuid, FWorldPartitionRuntimeCellObjectMapping>& SubActors = SourceActorsPerHLODLayer.FindOrAdd(HLODLayer);
-				AddSourceActor(ActorDescView, WorldPackageName, SubActors);
+				FPerHLODLayerData& PerHLODLayerData = PerHLODLayerDataMap.FindOrAdd(HLODLayer);
+				AddSourceActor(ActorDescView, WorldPackageName, PerHLODLayerData);
 			}
 		}
 	}
 
 	TArray<AWorldPartitionHLOD*> HLODActors;
-	for (auto& Pair : SourceActorsPerHLODLayer)
+	for (auto& Pair : PerHLODLayerDataMap)
 	{
 		const UHLODLayer* HLODLayer = Pair.Key;
-		TMap<FGuid, FWorldPartitionRuntimeCellObjectMapping>& SourceActorsMap = Pair.Value;
+		const FBox& Bounds = Pair.Value.Bounds;
+		TMap<FGuid, FWorldPartitionRuntimeCellObjectMapping>& SourceActorsMap = Pair.Value.SourceActors;
 		check(!SourceActorsMap.IsEmpty());
 
 		uint64 HLODActorHash = ComputeHLODActorUniqueHash(HLODLayer, InCreationParams.CellGuid);
@@ -359,11 +368,11 @@ TArray<AWorldPartitionHLOD*> FWorldPartitionHLODUtilities::CreateHLODActors(FHLO
 			DirtyReason = TEXT("FolderPath");
 		}
 
-		// Cell bounds
-		if (!HLODActor->GetHLODBounds().Equals(InCreationParams.CellBounds))
+		// HLOD bounds
+		if (!HLODActor->GetHLODBounds().Equals(Bounds))
 		{
-			HLODActor->SetHLODBounds(InCreationParams.CellBounds);
-			DirtyReason = TEXT("CellBounds");
+			HLODActor->SetHLODBounds(Bounds);
+			DirtyReason = TEXT("HLODBounds");
 		}
 
 		// Minimum visible distance
diff --git a/Engine/Source/Runtime/Engine/Private/WorldPartition/RuntimeHashSet/WorldPartitionRuntimeHashSetHLODGeneration.cpp b/Engine/Source/Runtime/Engine/Private/WorldPartition/RuntimeHashSet/WorldPartitionRuntimeHashSetHLODGeneration.cpp
--- a/Engine/Source/Runtime/Engine/Private/WorldPartition/RuntimeHashSet/WorldPartitionRuntimeHashSetHLODGeneration.cpp
+++ b/Engine/Source/Runtime/Engine/Private/WorldPartition/RuntimeHashSet/WorldPartitionRuntimeHashSetHLODGeneration.cpp
@@ -342,13 +342,12 @@ bool UWorldPartitionRuntimeHashSet::SetupHLODActors(const IStreamingGenerationCo
 
 				UE_LOG(LogWorldPartition, Display, TEXT("[%d / %d] %s %s..."), ++CellDescInstanceIndex, CellDescInstances.Num(), *LOCTEXT("ProcessingCell", "Processing cell").ToString(), *CellUniqueId.Name);
 
-				FBox CellBounds;
 				TArray<IStreamingGenerationContext::FActorInstance> ActorInstances;
 				for (const IStreamingGenerationContext::FActorSetInstance* ActorSetInstance : CellDescInstance.ActorSetInstances)
 				{
-					ActorSetInstance->ForEachActor([this, ActorSetInstance, &ActorInstances, &CellBounds](const FGuid& ActorGuid)
+					ActorSetInstance->ForEachActor([this, ActorSetInstance, &ActorInstances](const FGuid& ActorGuid)
 					{
-						CellBounds += ActorInstances.Emplace_GetRef(ActorGuid, ActorSetInstance).GetBounds();
+						ActorInstances.Emplace(ActorGuid, ActorSetInstance);
 					});
 				}
 
@@ -382,7 +381,7 @@ bool UWorldPartitionRuntimeHashSet::SetupHLODActors(const IStreamingGenerationCo
 				HLODCreationParams.TargetWorld = WorldPartition->GetWorld();
 				HLODCreationParams.CellName = CellUniqueId.Name;
 				HLODCreationParams.CellGuid = CellUniqueId.Guid;
-				HLODCreationParams.CellBounds = CellBounds;
+				HLODCreationParams.CellBounds = CellDescInstance.CellBounds.GetValue();
 				HLODCreationParams.GetRuntimeGrid = GetHLODRuntimeGrid;
 				HLODCreationParams.HLODLevel = HLODLevel;
 				HLODCreationParams.MinVisibleDistance = RuntimePartition->LoadingRange;
