[CustomHLOD] Update HLOD Actors creation to handle actors with Custom HLOD sources separately

In SetupHLODActors, source actors are processed in multiple waves:
* In the first wave, all regular actors are processed.
* In subsequent waves, HLOD actors generated in earlier steps are procesed.

This can cause issues when Custom HLOD Actors with a Parent Layer set are involved. They are processed during the first wave, but other contributing HLOD actors (contributing to the same target HLOD Actor) might be processed in later waves.
That breaks the expectation that each HLOD Actor and its sources are processed within a single wave.

To resolve this, Custom HLOD Actors are now used as source actors for separate HLOD Actors.

diff --git a/Engine/Plugins/Editor/WorldPartitionHLODUtilities/Source/Private/WorldPartition/HLOD/Utilities/WorldPartitionHLODUtilities.cpp b/Engine/Plugins/Editor/WorldPartitionHLODUtilities/Source/Private/WorldPartition/HLOD/Utilities/WorldPartitionHLODUtilities.cpp
--- a/Engine/Plugins/Editor/WorldPartitionHLODUtilities/Source/Private/WorldPartition/HLOD/Utilities/WorldPartitionHLODUtilities.cpp
+++ b/Engine/Plugins/Editor/WorldPartitionHLODUtilities/Source/Private/WorldPartition/HLOD/Utilities/WorldPartitionHLODUtilities.cpp
@@ -9,6 +9,7 @@
 #include "WorldPartition/WorldPartitionActorDescInstance.h"
 #include "WorldPartition/ContentBundle/ContentBundleActivationScope.h"
 #include "WorldPartition/DataLayer/DataLayerInstance.h"
+#include "WorldPartition/HLOD/CustomHLODActor.h"
 #include "WorldPartition/HLOD/HLODActor.h"
 #include "WorldPartition/HLOD/HLODHashBuilder.h"
 #include "WorldPartition/HLOD/HLODLayer.h"
@@ -78,6 +79,27 @@ static void ComputeHLODHash(FHLODHashBuilder& InHashBuilder, AWorldPartitionHLOD
 	UHLODBuilder::ComputeHLODHash(InHashBuilder, InSourceComponents);
 }
 
+struct FPerHLODLayerKey
+{
+	UHLODLayer* Layer;
+	bool bCustomHLOD;
+
+	FPerHLODLayerKey(UHLODLayer* InLayer, bool bInCustomHLOD)
+		: Layer(InLayer)
+		, bCustomHLOD(bInCustomHLOD)
+	{}
+
+	bool operator==(const FPerHLODLayerKey& Other) const
+	{
+		return Layer == Other.Layer && bCustomHLOD == Other.bCustomHLOD;
+	}
+
+	friend uint32 GetTypeHash(const FPerHLODLayerKey& Key)
+	{
+		return HashCombine(GetTypeHash(Key.Layer), GetTypeHash(Key.bCustomHLOD));
+	}
+};
+
 struct FPerHLODLayerData
 {
 	TMap<FGuid, FWorldPartitionRuntimeCellObjectMapping> SourceActors;
@@ -148,7 +170,7 @@ void AddSourceActor(const FStreamingGenerationActorDescView& ActorDescView, cons
 	}
 }
 
-uint64 ComputeHLODActorUniqueHash(const UHLODLayer* HLODLayer, const FGuid CellGuid)
+uint64 ComputeHLODActorUniqueHash(const UHLODLayer* HLODLayer, const FGuid CellGuid, bool bCustomHLOD)
 {
 	uint64 HLODActorHash = 0;
 	if (!CVarHLODUseLegacy32BitNameHash.GetValueOnGameThread())
@@ -161,6 +183,11 @@ uint64 ComputeHLODActorUniqueHash(const UHLODLayer* HLODLayer, const FGuid CellG
 			const FString HLODActorClassPathName = HLODLayer->GetHLODActorClass()->GetPathName();
 			HLODActorHash = CityHash64WithSeed(reinterpret_cast<const char*>(HLODActorClassPathName.GetCharArray().GetData()), HLODActorClassPathName.GetCharArray().Num() * sizeof(FString::ElementType), HLODActorHash);
 		}
+		if (bCustomHLOD)
+		{
+			static const FString CustomHLODText(TEXT("CustomHLOD"));
+			HLODActorHash = CityHash64WithSeed(reinterpret_cast<const char*>(CustomHLODText.GetCharArray().GetData()), CustomHLODText.GetCharArray().Num() * sizeof(FString::ElementType), HLODActorHash);
+		}
 	}
 	else
 	{
@@ -172,6 +199,12 @@ uint64 ComputeHLODActorUniqueHash(const UHLODLayer* HLODLayer, const FGuid CellG
 			const uint32 HLODActorClassHash = FCrc::StrCrc32(*HLODLayer->GetHLODActorClass()->GetPathName());
 			HLODActorHash32 = HashCombineFast(HLODActorHash32, HLODActorClassHash);
 		}
+		if (bCustomHLOD)
+		{
+			static const FString CustomHLODText(TEXT("CustomHLOD"));
+			const uint32 CustomHLODTextHash = GetTypeHash(CustomHLODText);
+			HLODActorHash32 = HashCombineFast(HLODActorHash32, CustomHLODTextHash);
+		}
 		HLODActorHash = HLODActorHash32;
 	}
 	return HLODActorHash;			
@@ -179,7 +212,7 @@ uint64 ComputeHLODActorUniqueHash(const UHLODLayer* HLODLayer, const FGuid CellG
 
 TArray<AWorldPartitionHLOD*> FWorldPartitionHLODUtilities::CreateHLODActors(FHLODCreationContext& InCreationContext, const FHLODCreationParams& InCreationParams, const TArray<IStreamingGenerationContext::FActorInstance>& InActors)
 {
-	TMap<UHLODLayer*, FPerHLODLayerData> PerHLODLayerDataMap;
+	TMap<FPerHLODLayerKey, FPerHLODLayerData> PerHLODLayerDataMap;
 
 	FName WorldPackageName = InCreationParams.WorldPartition->GetWorld()->GetPackage()->GetFName();
 
@@ -211,7 +244,8 @@ TArray<AWorldPartitionHLOD*> FWorldPartitionHLODUtilities::CreateHLODActors(FHLO
 					}
 				}
 
-				FPerHLODLayerData& PerHLODLayerData = PerHLODLayerDataMap.FindOrAdd(HLODLayer);
+				const bool bIsCustomHLODActor = ActorDescView.GetActorNativeClass()->IsChildOf<AWorldPartitionCustomHLOD>();
+				FPerHLODLayerData& PerHLODLayerData = PerHLODLayerDataMap.FindOrAdd(FPerHLODLayerKey(HLODLayer, bIsCustomHLODActor));
 				AddSourceActor(ActorDescView, WorldPackageName, PerHLODLayerData);
 			}
 		}
@@ -220,12 +254,13 @@ TArray<AWorldPartitionHLOD*> FWorldPartitionHLODUtilities::CreateHLODActors(FHLO
 	TArray<AWorldPartitionHLOD*> HLODActors;
 	for (auto& Pair : PerHLODLayerDataMap)
 	{
-		const UHLODLayer* HLODLayer = Pair.Key;
+		const UHLODLayer* HLODLayer = Pair.Key.Layer;
+		const bool bCustomHLOD = Pair.Key.bCustomHLOD;
 		const FBox& Bounds = Pair.Value.Bounds;
 		TMap<FGuid, FWorldPartitionRuntimeCellObjectMapping>& SourceActorsMap = Pair.Value.SourceActors;
 		check(!SourceActorsMap.IsEmpty());
 
-		uint64 HLODActorHash = ComputeHLODActorUniqueHash(HLODLayer, InCreationParams.CellGuid);
+		uint64 HLODActorHash = ComputeHLODActorUniqueHash(HLODLayer, InCreationParams.CellGuid, bCustomHLOD);
 		FName HLODActorName = *FString::Printf(TEXT("%s_%016" UINT64_x_FMT), *HLODLayer->GetName(), HLODActorHash);
 
 		AWorldPartitionHLOD* HLODActor = nullptr;
@@ -353,7 +388,7 @@ TArray<AWorldPartitionHLOD*> FWorldPartitionHLODUtilities::CreateHLODActors(FHLO
 		}
 
 		// Actor label
-		const FString ActorLabel = FString::Printf(TEXT("%s/%s"), *HLODLayer->GetName(), *InCreationParams.CellName);
+		const FString ActorLabel = FString::Printf(TEXT("%s/%s%s"), *HLODLayer->GetName(), *InCreationParams.CellName, bCustomHLOD ? TEXT("_CustomHLOD") : TEXT(""));
 		if (HLODActor->GetActorLabel() != ActorLabel)
 		{
 			HLODActor->SetActorLabel(ActorLabel);
