World Partition - HLOD: Added the ability to compute an updated HLOD hash without performing a rebuild

diff --git a/Engine/Plugins/Editor/WorldPartitionHLODUtilities/Source/Private/WorldPartition/HLOD/Utilities/WorldPartitionHLODUtilities.cpp b/Engine/Plugins/Editor/WorldPartitionHLODUtilities/Source/Private/WorldPartition/HLOD/Utilities/WorldPartitionHLODUtilities.cpp
--- a/Engine/Plugins/Editor/WorldPartitionHLODUtilities/Source/Private/WorldPartition/HLOD/Utilities/WorldPartitionHLODUtilities.cpp
+++ b/Engine/Plugins/Editor/WorldPartitionHLODUtilities/Source/Private/WorldPartition/HLOD/Utilities/WorldPartitionHLODUtilities.cpp
@@ -53,7 +53,7 @@ TAutoConsoleVariable<bool> CVarHLODUseLegacy32BitNameHash(
 	TEXT("Force the use of 32-bit hashing to use when generating HLOD actor names. Switching to 64-bit will result in new names/new actor files."),
 	ECVF_ReadOnly);
 
-static void ComputeHLODHash(FHLODHashBuilder& InHashBuilder, AWorldPartitionHLOD* InHLODActor, const TArray<UActorComponent*>& InSourceComponents)
+static void ComputeHLODHash(FHLODHashBuilder& InHashBuilder, const AWorldPartitionHLOD* InHLODActor, const TArray<UActorComponent*>& InSourceComponents)
 {
 	// Base key, changing this will force a rebuild of all HLODs
 	FString HLODBaseKey = TEXT("0D33837AB1A04CC2AEEB04C5C217DE96");
@@ -728,7 +728,7 @@ static TArray<UActorComponent*> GatherHLODRelevantComponents(const TArray<AActor
 	return HLODRelevantComponents.Array();
 }
 
-ULevelStreaming* LoadSourceActors(AWorldPartitionHLOD* InHLODActor, bool& bOutIsDirty)
+ULevelStreaming* LoadSourceActors(const AWorldPartitionHLOD* InHLODActor, bool& bOutIsDirty)
 {
 	// If we're loading actors for an HLOD > 0
 	// Ensure that async writes for source actors (which are HLODs N-1) are completed.
@@ -797,6 +797,11 @@ uint32 FWorldPartitionHLODUtilities::BuildHLOD(AWorldPartitionHLOD* InHLODActor)
 		FAutoScopedDurationTimer LoadTimeScope;
 
 		LevelStreaming = LoadSourceActors(InHLODActor, bIsDirty);
+		if (!ensure(LevelStreaming))
+		{
+			UE_LOG(LogHLODBuilder, Error, TEXT("FWorldPartitionHLODUtilities::BuildHLOD() - Failed to load source actors"));
+			return 0;
+		}
 
 		LoadTimeMS = FMath::RoundToInt(LoadTimeScope.GetTime() * 1000);
 	}
@@ -814,7 +819,7 @@ uint32 FWorldPartitionHLODUtilities::BuildHLOD(AWorldPartitionHLOD* InHLODActor)
 	}
 
 	FHLODHashBuilder HashBuilder;
-	ComputeHLODHash(HashBuilder, InHLODActor, HLODRelevantComponents);
+	::ComputeHLODHash(HashBuilder, InHLODActor, HLODRelevantComponents);
 
 	uint32 NewHLODHash = HashBuilder.GetCrc();
 	uint32 OldHLODHash = bIsDirty ? 0 : InHLODActor->GetHLODHash();
@@ -1036,6 +1041,33 @@ uint32 FWorldPartitionHLODUtilities::BuildHLOD(AWorldPartitionHLOD* InHLODActor)
 	return NewHLODHash;
 }
 
+uint32 FWorldPartitionHLODUtilities::ComputeHLODHash(const AWorldPartitionHLOD* InHLODActor)
+{
+	// Load actors relevant to HLODs
+	bool bIsDirty = false;
+	ULevelStreaming* LevelStreaming = LoadSourceActors(InHLODActor, bIsDirty);
+	if (!ensure(LevelStreaming))
+	{
+		UE_LOG(LogHLODBuilder, Error, TEXT("FWorldPartitionHLODUtilities::ComputeHLODHash() - Failed to load source actors"));
+		return 0;
+	}
+
+	ON_SCOPE_EXIT
+	{
+		UnloadSourceActors(LevelStreaming);
+	};
+
+	TArray<UActorComponent*> HLODRelevantComponents;
+	if (LevelStreaming->GetLoadedLevel())
+	{
+		HLODRelevantComponents = GatherHLODRelevantComponents(LevelStreaming->GetLoadedLevel()->Actors);
+	}
+
+	FHLODHashBuilder HashBuilder;
+	::ComputeHLODHash(HashBuilder, InHLODActor, HLODRelevantComponents);
+	return HashBuilder.GetCrc();
+}
+
 void FWorldPartitionHLODUtilities::SetHLODBuildEvaluator(IWorldPartitionHLODUtilities::FHLODBuildEvaluator BuildEvaluatorDelegate)
 {
 	HLODBuildEvaluatorDelegate = BuildEvaluatorDelegate;
diff --git a/Engine/Plugins/Editor/WorldPartitionHLODUtilities/Source/Public/WorldPartition/HLOD/Utilities/WorldPartitionHLODUtilities.h b/Engine/Plugins/Editor/WorldPartitionHLODUtilities/Source/Public/WorldPartition/HLOD/Utilities/WorldPartitionHLODUtilities.h
--- a/Engine/Plugins/Editor/WorldPartitionHLODUtilities/Source/Public/WorldPartition/HLOD/Utilities/WorldPartitionHLODUtilities.h
+++ b/Engine/Plugins/Editor/WorldPartitionHLODUtilities/Source/Public/WorldPartition/HLOD/Utilities/WorldPartitionHLODUtilities.h
@@ -12,6 +12,7 @@ class FWorldPartitionHLODUtilities : public IWorldPartitionHLODUtilities
 public:
 	virtual TArray<AWorldPartitionHLOD*> CreateHLODActors(FHLODCreationContext& InCreationContext, const FHLODCreationParams& InCreationParams, const TArray<IStreamingGenerationContext::FActorInstance>& InActors) override;
 	virtual uint32 BuildHLOD(AWorldPartitionHLOD* InHLODActor) override;
+	virtual uint32 ComputeHLODHash(const AWorldPartitionHLOD* InHLODActor) override;
 	virtual TSubclassOf<UHLODBuilder> GetHLODBuilderClass(const UHLODLayer* InHLODLayer) override;
 	virtual UHLODBuilderSettings* CreateHLODBuilderSettings(UHLODLayer* InHLODLayer) override;
 	virtual void SetHLODBuildEvaluator(IWorldPartitionHLODUtilities::FHLODBuildEvaluator BuildEvaluatorDelegate) override;
diff --git a/Engine/Source/Runtime/Engine/Private/WorldPartition/HLOD/HLODActor.cpp b/Engine/Source/Runtime/Engine/Private/WorldPartition/HLOD/HLODActor.cpp
--- a/Engine/Source/Runtime/Engine/Private/WorldPartition/HLOD/HLODActor.cpp
+++ b/Engine/Source/Runtime/Engine/Private/WorldPartition/HLOD/HLODActor.cpp
@@ -613,6 +613,17 @@ void AWorldPartitionHLOD::BuildHLOD(bool bForceBuild)
 	GetHLODBuildEventDelegate().Broadcast({ FWorldPartitionHLODBuildEvent::EEventType::EndBuild, this });
 }
 
+uint32 AWorldPartitionHLOD::ComputeHLODHash() const
+{
+	uint32 ComputedHLODHash = 0;
+	IWorldPartitionHLODUtilitiesModule* WPHLODUtilitiesModule = FModuleManager::Get().LoadModulePtr<IWorldPartitionHLODUtilitiesModule>("WorldPartitionHLODUtilities");
+	if (IWorldPartitionHLODUtilities* WPHLODUtilities = WPHLODUtilitiesModule != nullptr ? WPHLODUtilitiesModule->GetUtilities() : nullptr)
+	{
+		ComputedHLODHash = WPHLODUtilities->ComputeHLODHash(this);
+	}
+	return ComputedHLODHash;
+}
+
 FWorldPartitionHLODBuildEventDelegate& AWorldPartitionHLOD::GetHLODBuildEventDelegate()
 {
 	return HLODBuildEventDelegate;
diff --git a/Engine/Source/Runtime/Engine/Public/WorldPartition/HLOD/HLODActor.h b/Engine/Source/Runtime/Engine/Public/WorldPartition/HLOD/HLODActor.h
--- a/Engine/Source/Runtime/Engine/Public/WorldPartition/HLOD/HLODActor.h
+++ b/Engine/Source/Runtime/Engine/Public/WorldPartition/HLOD/HLODActor.h
@@ -116,6 +116,7 @@ public:
 	void SetMinVisibleDistance(double InMinVisibleDistance) { MinVisibleDistance = InMinVisibleDistance; }
 
 	ENGINE_API void BuildHLOD(bool bForceBuild = false);
+	ENGINE_API uint32 ComputeHLODHash() const;
 	ENGINE_API uint32 GetHLODHash() const;
 	ENGINE_API void SetHLODHash(uint32 InHLODHash, const FString& InHLODBuildReportContent);
 
diff --git a/Engine/Source/Runtime/Engine/Public/WorldPartition/HLOD/HLODProviderInterface.h b/Engine/Source/Runtime/Engine/Public/WorldPartition/HLOD/HLODProviderInterface.h
--- a/Engine/Source/Runtime/Engine/Public/WorldPartition/HLOD/HLODProviderInterface.h
+++ b/Engine/Source/Runtime/Engine/Public/WorldPartition/HLOD/HLODProviderInterface.h
@@ -29,7 +29,7 @@ public:
 
 	struct FBuildHLODActorParams
 	{
-		bool bForceRebuild;
+		bool bForceRebuild = false;
 
 		FPackageProcessedDelegate OnPackageProcessed;
 
diff --git a/Engine/Source/Runtime/Engine/Public/WorldPartition/HLOD/IWorldPartitionHLODUtilities.h b/Engine/Source/Runtime/Engine/Public/WorldPartition/HLOD/IWorldPartitionHLODUtilities.h
--- a/Engine/Source/Runtime/Engine/Public/WorldPartition/HLOD/IWorldPartitionHLODUtilities.h
+++ b/Engine/Source/Runtime/Engine/Public/WorldPartition/HLOD/IWorldPartitionHLODUtilities.h
@@ -81,6 +81,14 @@ public:
 	 */
 	virtual uint32 BuildHLOD(AWorldPartitionHLOD* InHLODActor) = 0;
 
+	/**
+	 * Compute the HLOD hash for the specified AWorldPartitionHLOD actor.
+	 *
+	 * @param 	InHLODActor		The HLOD actor for which we'll compute the hash
+	 * @return An hash that represent the content used to build this HLOD.
+	 */
+	virtual uint32 ComputeHLODHash(const AWorldPartitionHLOD* InHLODActor) = 0;
+
 	/**
 	 * Retrieve the HLOD Builder class to use for the given HLODLayer.
 	 * 
