World Partition - HLOD: Refactored Source Actors loading
* Now handles missing source actors & renamed levels
* Now loading actors in an isolated world to avoid any potential issues
* Deprecated UWorldPartitionLevelStreamingDynamic::LoadInEditor() in favor of FWorldPartitionLevelHelper::LoadActors()

diff --git a/Engine/Plugins/Editor/WorldPartitionHLODUtilities/Source/Private/WorldPartition/HLOD/Utilities/WorldPartitionHLODUtilities.cpp b/Engine/Plugins/Editor/WorldPartitionHLODUtilities/Source/Private/WorldPartition/HLOD/Utilities/WorldPartitionHLODUtilities.cpp
--- a/Engine/Plugins/Editor/WorldPartitionHLODUtilities/Source/Private/WorldPartition/HLOD/Utilities/WorldPartitionHLODUtilities.cpp
+++ b/Engine/Plugins/Editor/WorldPartitionHLODUtilities/Source/Private/WorldPartition/HLOD/Utilities/WorldPartitionHLODUtilities.cpp
@@ -696,19 +696,19 @@ void GatherOutputStats(AWorldPartitionHLOD* InHLODActor)
 }
 
 // Iterate over the source actors and retrieve HLOD relevant components using GetHLODRelevantComponents()
-static TArray<UActorComponent*> GatherHLODRelevantComponents(const TArray<AActor*>& InSourceActors)
+static TArray<UActorComponent*> GatherHLODRelevantComponents(UWorld* InWorld)
 {
 	TSet<UActorComponent*> HLODRelevantComponents;
 
-	for (AActor* Actor : InSourceActors)
+	for (TActorIterator<AActor> It(InWorld); It; ++It)
 	{
-		if (!Actor || !Actor->IsHLODRelevant())
+		if (!It->IsHLODRelevant())
 		{
 			continue;
 		}
 
 		// Extract components to be used as input for the HLOD generation process
-		for (UActorComponent* HLODRelevantComponentForActor : Actor->GetHLODRelevantComponents())
+		for (UActorComponent* HLODRelevantComponentForActor : It->GetHLODRelevantComponents())
 		{
 			// Components can return proxy components to be used in their place while building HLODs
 			TArray<UActorComponent*> HLODProxyComponents = HLODRelevantComponentForActor->GetHLODProxyComponents();
@@ -728,7 +728,7 @@ static TArray<UActorComponent*> GatherHLODRelevantComponents(const TArray<AActor
 	return HLODRelevantComponents.Array();
 }
 
-ULevelStreaming* LoadSourceActors(const AWorldPartitionHLOD* InHLODActor, bool& bOutIsDirty)
+bool LoadSourceActors(const AWorldPartitionHLOD* InHLODActor, UWorld* TargetWorld, bool& bOutIsDirty)
 {
 	// If we're loading actors for an HLOD > 0
 	// Ensure that async writes for source actors (which are HLODs N-1) are completed.
@@ -737,10 +737,10 @@ ULevelStreaming* LoadSourceActors(const AWorldPartitionHLOD* InHLODActor, bool&
 		UPackage::WaitForAsyncFileWrites();
 	}
 	
-	ULevelStreaming* LevelStreaming = InHLODActor->GetSourceActors()->LoadSourceActors(bOutIsDirty);
+	bool bLoadSuccess = InHLODActor->GetSourceActors()->LoadSourceActors(bOutIsDirty, TargetWorld);
 	UE_CLOG(bOutIsDirty, LogHLODBuilder, Warning, TEXT("HLOD actor \"%s\" needs to be rebuilt as it didn't succeed in loading all actors."), *InHLODActor->GetActorLabel());
 
-	if (LevelStreaming->GetLoadedLevel())
+	if (bLoadSuccess)
 	{
 		// Finish assets compilation
 		FAssetCompilingManager::Get().FinishAllCompilation();
@@ -759,7 +759,7 @@ ULevelStreaming* LoadSourceActors(const AWorldPartitionHLOD* InHLODActor, bool&
 		}
 	}
 
-	return LevelStreaming;
+	return bLoadSuccess;
 }
 
 void UnloadSourceActors(ULevelStreaming* InLevelStreaming)
@@ -792,12 +792,22 @@ uint32 FWorldPartitionHLODUtilities::BuildHLOD(AWorldPartitionHLOD* InHLODActor)
 	
 	// Load actors relevant to HLODs
 	bool bIsDirty = false;
-	ULevelStreaming* LevelStreaming = nullptr;
+
+	const bool bInformEngineOfWorld = false;
+	const bool bAddToRoot = true;
+	UWorld* TempWorld = UWorld::CreateWorld(EWorldType::Editor, bInformEngineOfWorld, MakeUniqueObjectName(GetTransientPackage(), UWorld::StaticClass(), TEXT("BuildHLODWorld")), GetTransientPackage(), bAddToRoot);
+	check(TempWorld);
+
+	ON_SCOPE_EXIT
+	{
+		TempWorld->DestroyWorld(bInformEngineOfWorld);
+	};
+
 	{
 		FAutoScopedDurationTimer LoadTimeScope;
 
-		LevelStreaming = LoadSourceActors(InHLODActor, bIsDirty);
-		if (!ensure(LevelStreaming))
+		bool bLoadSucceeded = LoadSourceActors(InHLODActor, TempWorld, bIsDirty);
+		if (!ensure(bLoadSucceeded))
 		{
 			UE_LOG(LogHLODBuilder, Error, TEXT("FWorldPartitionHLODUtilities::BuildHLOD() - Failed to load source actors"));
 			return 0;
@@ -806,18 +816,7 @@ uint32 FWorldPartitionHLODUtilities::BuildHLOD(AWorldPartitionHLOD* InHLODActor)
 		LoadTimeMS = FMath::RoundToInt(LoadTimeScope.GetTime() * 1000);
 	}
 
-	ON_SCOPE_EXIT
-	{
-		UnloadSourceActors(LevelStreaming);
-	};
-
-	TArray<UActorComponent*> HLODRelevantComponents;
-	
-	if (LevelStreaming->GetLoadedLevel())
-	{
-		HLODRelevantComponents = GatherHLODRelevantComponents(LevelStreaming->GetLoadedLevel()->Actors);
-	}
-
+	TArray<UActorComponent*> HLODRelevantComponents = GatherHLODRelevantComponents(TempWorld);
 	FHLODHashBuilder HashBuilder;
 	::ComputeHLODHash(HashBuilder, InHLODActor, HLODRelevantComponents);
 
@@ -1043,26 +1042,27 @@ uint32 FWorldPartitionHLODUtilities::BuildHLOD(AWorldPartitionHLOD* InHLODActor)
 
 uint32 FWorldPartitionHLODUtilities::ComputeHLODHash(const AWorldPartitionHLOD* InHLODActor)
 {
-	// Load actors relevant to HLODs
-	bool bIsDirty = false;
-	ULevelStreaming* LevelStreaming = LoadSourceActors(InHLODActor, bIsDirty);
-	if (!ensure(LevelStreaming))
-	{
-		UE_LOG(LogHLODBuilder, Error, TEXT("FWorldPartitionHLODUtilities::ComputeHLODHash() - Failed to load source actors"));
-		return 0;
-	}
+	const bool bInformEngineOfWorld = false;
+	const bool bAddToRoot = true;
+	UWorld* TempWorld = UWorld::CreateWorld(EWorldType::Editor, bInformEngineOfWorld, MakeUniqueObjectName(GetTransientPackage(), UWorld::StaticClass(), TEXT("BuildHLODWorld")), GetTransientPackage(), bAddToRoot);
+	check(TempWorld);
 
 	ON_SCOPE_EXIT
 	{
-		UnloadSourceActors(LevelStreaming);
+		TempWorld->DestroyWorld(bInformEngineOfWorld);
 	};
 
-	TArray<UActorComponent*> HLODRelevantComponents;
-	if (LevelStreaming->GetLoadedLevel())
+	// Load actors relevant to HLODs
+	bool bIsDirty = false;
+	bool bLoadSucceeded = LoadSourceActors(InHLODActor, TempWorld, bIsDirty);
+	if (!ensure(bLoadSucceeded))
 	{
-		HLODRelevantComponents = GatherHLODRelevantComponents(LevelStreaming->GetLoadedLevel()->Actors);
+		UE_LOG(LogHLODBuilder, Error, TEXT("FWorldPartitionHLODUtilities::ComputeHLODHash() - Failed to load source actors"));
+		return 0;
 	}
 
+	TArray<UActorComponent*> HLODRelevantComponents = GatherHLODRelevantComponents(TempWorld);
+
 	FHLODHashBuilder HashBuilder;
 	::ComputeHLODHash(HashBuilder, InHLODActor, HLODRelevantComponents);
 	return HashBuilder.GetCrc();
diff --git a/Engine/Source/Runtime/Engine/Private/WorldPartition/HLOD/HLODSourceActorsFromCell.cpp b/Engine/Source/Runtime/Engine/Private/WorldPartition/HLOD/HLODSourceActorsFromCell.cpp
--- a/Engine/Source/Runtime/Engine/Private/WorldPartition/HLOD/HLODSourceActorsFromCell.cpp
+++ b/Engine/Source/Runtime/Engine/Private/WorldPartition/HLOD/HLODSourceActorsFromCell.cpp
@@ -7,6 +7,11 @@
 #include "WorldPartition/HLOD/HLODHashBuilder.h"
 #include "Serialization/ArchiveCrc32.h"
 #include "WorldPartition/WorldPartitionLevelStreamingDynamic.h"
+
+#include "Editor/UnrealEd/Public/EditorLevelUtils.h" // MoveActorsToLevel
+#include "WorldPartition/ContentBundle/ContentBundlePaths.h"
+#include "WorldPartition/DataLayer/ExternalDataLayerHelper.h"
+
 #endif
 
 #include UE_INLINE_GENERATED_CPP_BY_NAME(HLODSourceActorsFromCell)
@@ -18,27 +23,61 @@ UWorldPartitionHLODSourceActorsFromCell::UWorldPartitionHLODSourceActorsFromCell
 
 #if WITH_EDITOR
 
-ULevelStreaming* UWorldPartitionHLODSourceActorsFromCell::LoadSourceActors(bool& bOutDirty) const
+bool UWorldPartitionHLODSourceActorsFromCell::LoadSourceActors(bool& bOutDirty, UWorld* TargetWorld) const
 {
 	UPackage::WaitForAsyncFileWrites();
 
 	bOutDirty = false;
 	AWorldPartitionHLOD* HLODActor = CastChecked<AWorldPartitionHLOD>(GetOuter());
-	UWorld* World = HLODActor->GetWorld();
+	UWorld* SourceWorld = HLODActor->GetWorld();
+
+	FLinkerInstancingContext InstancingContext;
+
+	// Don't do SoftObjectPath remapping for PersistentLevel actors because references can end up in different cells
+	InstancingContext.SetSoftObjectPathRemappingEnabled(false);
 
-	const FName LevelStreamingName = FName(*FString::Printf(TEXT("HLODLevelStreaming_%s"), *HLODActor->GetName()));
+	UPackage* SourcePackage = SourceWorld->GetPackage();
+	UPackage* TargetPackage = TargetWorld->GetPackage();
+
+	FName SourcePackageName = SourcePackage->GetFName();
+	FName TargetPackageName = TargetPackage->GetFName();
+	InstancingContext.AddPackageMapping(SourcePackageName, TargetPackageName);
+
+	for (const FWorldPartitionRuntimeCellObjectMapping& CellObjectMapping : Actors)
+	{
+		if (ContentBundlePaths::IsAContentBundlePath(CellObjectMapping.ContainerPackage.ToString()) ||
+			FExternalDataLayerHelper::IsExternalDataLayerPath(CellObjectMapping.ContainerPackage.ToString()))
+		{
+			check(CellObjectMapping.ContainerPackage != CellObjectMapping.WorldPackage);
+			bool bIsContainerPackageAlreadyRemapped = InstancingContext.RemapPackage(CellObjectMapping.ContainerPackage) != CellObjectMapping.ContainerPackage;
+			if (!bIsContainerPackageAlreadyRemapped)
+			{
+				InstancingContext.AddPackageMapping(CellObjectMapping.ContainerPackage, TargetPackageName);
+			}
+		}
+	}
 
-	UWorldPartitionLevelStreamingDynamic* LevelStreaming = UWorldPartitionLevelStreamingDynamic::LoadInEditor(World, LevelStreamingName, Actors);
-	check(LevelStreaming);
+	TArray<FWorldPartitionRuntimeCellObjectMapping> ActorsCopy = Actors;
 
-	LevelStreaming->GetWorldPartitionRuntimeCell()->RuntimeCellData->DebugName = HLODActor->GetActorNameOrLabel();
+	FWorldPartitionLevelHelper::FPackageReferencer PackageReferencer;
+	FWorldPartitionLevelHelper::FLoadActorsParams Params = FWorldPartitionLevelHelper::FLoadActorsParams()
+		.SetOuterWorld(TargetWorld)
+		.SetDestLevel(nullptr)
+		.SetActorPackages(ActorsCopy)
+		.SetPackageReferencer(&PackageReferencer)
+		.SetCompletionCallback([&bOutDirty](bool bSucceeded) { bOutDirty = !bSucceeded; })
+		.SetLoadAsync(false)
+		.SetInstancingContext(MoveTemp(InstancingContext))
+		.SetSilenceLoadFailures(true);
 
-	if (!LevelStreaming->GetLoadSucceeded())
+	if (FWorldPartitionLevelHelper::LoadActors(MoveTemp(Params)))
 	{
-		bOutDirty = true;
+		TArray<UPackage*> ModifiedPackages;
+		FWorldPartitionLevelHelper::MoveExternalActorsToLevel(ActorsCopy, TargetWorld->PersistentLevel, ModifiedPackages);
+		return true;
 	}
 
-	return LevelStreaming;
+	return false;
 }
 
 uint32 UWorldPartitionHLODSourceActorsFromCell::GetSourceActorsHash(const TArray<FWorldPartitionRuntimeCellObjectMapping>& InSourceActors)
diff --git a/Engine/Source/Runtime/Engine/Private/WorldPartition/HLOD/HLODSourceActorsFromLevel.cpp b/Engine/Source/Runtime/Engine/Private/WorldPartition/HLOD/HLODSourceActorsFromLevel.cpp
--- a/Engine/Source/Runtime/Engine/Private/WorldPartition/HLOD/HLODSourceActorsFromLevel.cpp
+++ b/Engine/Source/Runtime/Engine/Private/WorldPartition/HLOD/HLODSourceActorsFromLevel.cpp
@@ -18,22 +18,12 @@ UWorldPartitionHLODSourceActorsFromLevel::UWorldPartitionHLODSourceActorsFromLev
 
 #if WITH_EDITOR
 
-ULevelStreaming* UWorldPartitionHLODSourceActorsFromLevel::LoadSourceActors(bool& bOutDirty) const
+bool UWorldPartitionHLODSourceActorsFromLevel::LoadSourceActors(bool& bOutDirty, UWorld* TargetWorld) const
 {
-	AWorldPartitionHLOD* HLODActor = CastChecked<AWorldPartitionHLOD>(GetOuter());
-	UWorld* World = HLODActor->GetWorld();
-
 	bool bSuccess = false;
-	FName LevelStreamingName = SourceLevel.GetUniqueID().GetLongPackageFName();
-	ULevelStreamingDynamic* LevelStreaming = ULevelStreamingDynamic::LoadLevelInstanceBySoftObjectPtr(World, SourceLevel, FTransform::Identity, bSuccess, TEXT(""), nullptr, true);
-
-	if (bSuccess && LevelStreaming)
-	{
-		World->FlushLevelStreaming();
-		return LevelStreaming;
-	}
-
-	return nullptr;
+	ULevelStreamingDynamic* LevelStreaming = ULevelStreamingDynamic::LoadLevelInstanceBySoftObjectPtr(TargetWorld, SourceLevel, FTransform::Identity, bSuccess);
+	TargetWorld->FlushLevelStreaming();
+	return bSuccess && LevelStreaming && LevelStreaming->GetLoadedLevel();
 }
 
 void UWorldPartitionHLODSourceActorsFromLevel::ComputeHLODHash(FHLODHashBuilder& InHashBuilder) const
diff --git a/Engine/Source/Runtime/Engine/Private/WorldPartition/WorldPartitionLevelHelper.cpp b/Engine/Source/Runtime/Engine/Private/WorldPartition/WorldPartitionLevelHelper.cpp
--- a/Engine/Source/Runtime/Engine/Private/WorldPartition/WorldPartitionLevelHelper.cpp
+++ b/Engine/Source/Runtime/Engine/Private/WorldPartition/WorldPartitionLevelHelper.cpp
@@ -611,6 +611,7 @@ bool FWorldPartitionLevelHelper::LoadActors(const FLoadActorsParams& InParams)
 		.SetDestLevel(InParams.DestLevel)
 		.SetInstancingContext(MoveTemp(InParams.InstancingContext))
 		.SetLoadAsync(InParams.bLoadAsync, InParams.AsyncRequestIDs)
+		.SetSilenceLoadFailures(InParams.bSilenceLoadFailures)
 		.SetOuterWorld(InParams.OuterWorld)
 		.SetPackageReferencer(InParams.PackageReferencer);
 	return LoadActors(MoveTemp(ParamsCopy));
@@ -732,7 +733,7 @@ bool FWorldPartitionLevelHelper::LoadActorsInternal(FLoadActorsParams&& InParams
 			}
 		}
 
-		FLoadPackageAsyncDelegate CompletionCallback = FLoadPackageAsyncDelegate::CreateLambda([LoadProgress, PackageObjectMapping, LoadedOverrides = InLoadedPropertyOverrides, PackageReferencer = InParams.PackageReferencer, OuterWorld = InParams.OuterWorld, DestLevel = InParams.DestLevel, CompletionCallback = InParams.CompletionCallback](const FName& LoadedPackageName, UPackage* LoadedPackage, EAsyncLoadingResult::Type Result)
+		FLoadPackageAsyncDelegate CompletionCallback = FLoadPackageAsyncDelegate::CreateLambda([LoadProgress, PackageObjectMapping, LoadedOverrides = InLoadedPropertyOverrides, PackageReferencer = InParams.PackageReferencer, OuterWorld = InParams.OuterWorld, DestLevel = InParams.DestLevel, CompletionCallback = InParams.CompletionCallback, bSilenceLoadFailures = InParams.bSilenceLoadFailures](const FName& LoadedPackageName, UPackage* LoadedPackage, EAsyncLoadingResult::Type Result)
 		{
 			const FName ActorName = *FPaths::GetExtension(PackageObjectMapping->Path.ToString());
 			check(LoadProgress->NumPendingLoadRequests);
@@ -837,9 +838,6 @@ bool FWorldPartitionLevelHelper::LoadActorsInternal(FLoadActorsParams&& InParams
 					// Set the actor's instance guid
 					FSetActorInstanceGuid SetActorInstanceGuid(Actor, PackageObjectMapping->ActorInstanceGuid);
 						
-					// Path to use when searching for this actor in MoveExternalActorsToLevel
-					PackageObjectMapping->LoadedPath = *Actor->GetPathName();
-
 					// Fixup any FSoftObjectPath from this Actor (and its SubObjects) in this container to another object in the same container with a ContainerID suffix that can be remapped to
 					// a Cell package in the StreamingPolicy.
 					// 
@@ -868,6 +866,9 @@ bool FWorldPartitionLevelHelper::LoadActorsInternal(FLoadActorsParams&& InParams
 					FLevelUtils::ApplyLevelTransform(TransformParams);
 				}
 
+				// Path to use when searching for this actor in MoveExternalActorsToLevel
+				PackageObjectMapping->LoadedPath = *Actor->GetPathName();
+
 				if (DestLevel)
 				{
 					// Propagate resolved actor folder path
@@ -895,17 +896,20 @@ bool FWorldPartitionLevelHelper::LoadActorsInternal(FLoadActorsParams&& InParams
 			}
 			else
 			{
-				if (LoadedPackage)
+				if (!bSilenceLoadFailures)
 				{
-					UE_LOG(LogWorldPartition, Warning, TEXT("\tPackage Content for '%s:"), *LoadedPackage->GetName());
-					ForEachObjectWithOuter(LoadedPackage, [](UObject* Object)
+					if (LoadedPackage)
 					{
-						UE_LOG(LogWorldPartition, Warning, TEXT("\t\tObject %s, Flags 0x%llx"), *Object->GetPathName(), static_cast<uint64>(Object->GetFlags()));
-						return true;
-					}, true);
-				}
+						UE_LOG(LogWorldPartition, Warning, TEXT("\tPackage Content for '%s:"), *LoadedPackage->GetName());
+						ForEachObjectWithOuter(LoadedPackage, [](UObject* Object)
+						{
+							UE_LOG(LogWorldPartition, Warning, TEXT("\t\tObject %s, Flags 0x%llx"), *Object->GetPathName(), static_cast<uint64>(Object->GetFlags()));
+							return true;
+						}, true);
+					}
 
-				ensureMsgf(false, TEXT("Failed to find actor '%s' in package '%s'."), *ActorName.ToString(), *LoadedPackageName.ToString());
+					ensureMsgf(false, TEXT("Failed to find actor '%s' in package '%s'."), *ActorName.ToString(), *LoadedPackageName.ToString());
+				}
 
 				LoadProgress->NumFailedLoadedRequests++;
 			}
diff --git a/Engine/Source/Runtime/Engine/Private/WorldPartition/WorldPartitionLevelStreamingDynamic.cpp b/Engine/Source/Runtime/Engine/Private/WorldPartition/WorldPartitionLevelStreamingDynamic.cpp
--- a/Engine/Source/Runtime/Engine/Private/WorldPartition/WorldPartitionLevelStreamingDynamic.cpp
+++ b/Engine/Source/Runtime/Engine/Private/WorldPartition/WorldPartitionLevelStreamingDynamic.cpp
@@ -313,9 +313,11 @@ bool UWorldPartitionLevelStreamingDynamic::RequestLevel(UWorld* InPersistentWorl
 			check(CellLevelPackage);
 			check(UWorld::FindWorldInPackage(CellLevelPackage));
 			check(RuntimeLevel->OwningWorld);
-			check(RuntimeLevel->OwningWorld->WorldType == EWorldType::PIE || 
-				((IsRunningGame() || IsRunningDedicatedServer()) && RuntimeLevel->OwningWorld->WorldType == EWorldType::Game) || 
-				(RuntimeLevel->OwningWorld->WorldType == EWorldType::Editor && GetShouldBeVisibleInEditor()));
+
+			EWorldType::Type WorldType = RuntimeLevel->OwningWorld->WorldType;
+			check(WorldType == EWorldType::PIE ||
+				((IsRunningGame() || IsRunningDedicatedServer()) && WorldType == EWorldType::Game) ||
+				((WorldType == EWorldType::Editor || WorldType == EWorldType::EditorPreview) && GetShouldBeVisibleInEditor()));
 
 			if (IssueLoadRequests())
 			{
diff --git a/Engine/Source/Runtime/Engine/Public/WorldPartition/HLOD/HLODSourceActors.h b/Engine/Source/Runtime/Engine/Public/WorldPartition/HLOD/HLODSourceActors.h
--- a/Engine/Source/Runtime/Engine/Public/WorldPartition/HLOD/HLODSourceActors.h
+++ b/Engine/Source/Runtime/Engine/Public/WorldPartition/HLOD/HLODSourceActors.h
@@ -17,7 +17,7 @@ class UWorldPartitionHLODSourceActors : public UObject
 
 public:
 #if WITH_EDITOR
-	virtual ULevelStreaming* LoadSourceActors(bool& bOutDirty) const PURE_VIRTUAL(UWorldPartitionHLODSourceActors::LoadSourceActors, return nullptr; );
+	virtual bool LoadSourceActors(bool& bOutDirty, UWorld* TargetWorld) const PURE_VIRTUAL(UWorldPartitionHLODSourceActors::LoadSourceActors, return false; );
 	ENGINE_API virtual void ComputeHLODHash(FHLODHashBuilder& HashBuilder) const;
 
 	ENGINE_API void SetHLODLayer(const UHLODLayer* HLODLayer);
@@ -25,6 +25,9 @@ public:
 
 	UE_DEPRECATED(5.7, "GetHLODHash() has been replaced by ComputeHLODHash()")
 	virtual uint32 GetHLODHash() const final { return 0; }
+
+	UE_DEPRECATED(5.7, "Use LoadSourceActors() override with a TargetWorld instead")
+	virtual ULevelStreaming* LoadSourceActors(bool& bOutDirty) const final { return nullptr; }
 #endif
 
 private:
diff --git a/Engine/Source/Runtime/Engine/Public/WorldPartition/HLOD/HLODSourceActorsFromCell.h b/Engine/Source/Runtime/Engine/Public/WorldPartition/HLOD/HLODSourceActorsFromCell.h
--- a/Engine/Source/Runtime/Engine/Public/WorldPartition/HLOD/HLODSourceActorsFromCell.h
+++ b/Engine/Source/Runtime/Engine/Public/WorldPartition/HLOD/HLODSourceActorsFromCell.h
@@ -14,7 +14,7 @@ class UWorldPartitionHLODSourceActorsFromCell : public UWorldPartitionHLODSource
 
 public:
 #if WITH_EDITOR
-	ENGINE_API virtual ULevelStreaming* LoadSourceActors(bool& bOutDirty) const override;
+	ENGINE_API virtual bool LoadSourceActors(bool& bOutDirty, UWorld* TargetWorld) const override;
 	ENGINE_API virtual void ComputeHLODHash(FHLODHashBuilder& InHashBuilder) const override;
 
 	ENGINE_API static uint32 GetSourceActorsHash(const TArray<FWorldPartitionRuntimeCellObjectMapping>& InSourceActors);
diff --git a/Engine/Source/Runtime/Engine/Public/WorldPartition/HLOD/HLODSourceActorsFromLevel.h b/Engine/Source/Runtime/Engine/Public/WorldPartition/HLOD/HLODSourceActorsFromLevel.h
--- a/Engine/Source/Runtime/Engine/Public/WorldPartition/HLOD/HLODSourceActorsFromLevel.h
+++ b/Engine/Source/Runtime/Engine/Public/WorldPartition/HLOD/HLODSourceActorsFromLevel.h
@@ -13,7 +13,7 @@ class UWorldPartitionHLODSourceActorsFromLevel : public UWorldPartitionHLODSourc
 
 public:
 #if WITH_EDITOR
-	ENGINE_API virtual ULevelStreaming* LoadSourceActors(bool& bOutDirty) const override;
+	ENGINE_API virtual bool LoadSourceActors(bool& bOutDirty, UWorld* TargetWorld) const override;
 	ENGINE_API virtual void ComputeHLODHash(FHLODHashBuilder& InHashBuilder) const override;
 
 	ENGINE_API void SetSourceLevel(const UWorld* InSourceLevel);
diff --git a/Engine/Source/Runtime/Engine/Public/WorldPartition/HLOD/IWorldPartitionHLODUtilities.h b/Engine/Source/Runtime/Engine/Public/WorldPartition/HLOD/IWorldPartitionHLODUtilities.h
--- a/Engine/Source/Runtime/Engine/Public/WorldPartition/HLOD/IWorldPartitionHLODUtilities.h
+++ b/Engine/Source/Runtime/Engine/Public/WorldPartition/HLOD/IWorldPartitionHLODUtilities.h
@@ -112,12 +112,14 @@ public:
 	 * @param NewHash			The newly computed hash of the inputs to the HLOD build.
 	 * @return true if the HLOD build should be performed, false otherwise.
 	 */
+	UE_EXPERIMENTAL(5.7, "Experimental API, HLOD build evaluator logic may be moved in a future release.")
 	DECLARE_DELEGATE_RetVal_ThreeParams(bool, FHLODBuildEvaluator, AWorldPartitionHLOD*, uint32, uint32);
 
 	/**
 	 * Provide a delegate that will be used to evaluate whether an HLOD build should be performed.
 	 * @param The evaluator delegate.
 	 */
+	UE_EXPERIMENTAL(5.7, "Experimental API, HLOD build evaluator logic may be moved in a future release.")
 	virtual void SetHLODBuildEvaluator(FHLODBuildEvaluator BuildEvaluatorDelegate) = 0;
 
 
diff --git a/Engine/Source/Runtime/Engine/Public/WorldPartition/WorldPartitionLevelHelper.h b/Engine/Source/Runtime/Engine/Public/WorldPartition/WorldPartitionLevelHelper.h
--- a/Engine/Source/Runtime/Engine/Public/WorldPartition/WorldPartitionLevelHelper.h
+++ b/Engine/Source/Runtime/Engine/Public/WorldPartition/WorldPartitionLevelHelper.h
@@ -61,6 +61,7 @@ public:
 		FPackageReferencer* PackageReferencer;
 		TFunction<void(bool)> CompletionCallback;
 		bool bLoadAsync;
+		bool bSilenceLoadFailures;
 		TArray<int32>* AsyncRequestIDs;
 		mutable FLinkerInstancingContext InstancingContext;
 
@@ -75,6 +76,7 @@ public:
 			AsyncRequestIDs = OutAsyncRequestIDs;
 			return *this;
 		}
+		FLoadActorsParams& SetSilenceLoadFailures(bool bInSilenceLoadFailures) { bSilenceLoadFailures = bInSilenceLoadFailures; return *this; }
 		FLoadActorsParams& SetInstancingContext(FLinkerInstancingContext InInstancingContext) { InstancingContext = InInstancingContext; return *this; }
 	};
 
@@ -120,6 +122,7 @@ private:
 	friend class UWorldPartitionLevelStreamingDynamic;
 	friend class UWorldPartitionSubsystem;
 	friend class UWorldPartitionRuntimeLevelStreamingCell;
+	friend class UWorldPartitionHLODSourceActorsFromCell;
 	// Cache of Property Overrides to apply after ReRunConstructionScript
 	// In PIE this will  be done when Streaming state changes to Visible on the UWorldPartitionLevelStreamingDynamic
 	// In Cook this will be done on save of the Level cell
diff --git a/Engine/Source/Runtime/Engine/Public/WorldPartition/WorldPartitionLevelStreamingDynamic.h b/Engine/Source/Runtime/Engine/Public/WorldPartition/WorldPartitionLevelStreamingDynamic.h
--- a/Engine/Source/Runtime/Engine/Public/WorldPartition/WorldPartitionLevelStreamingDynamic.h
+++ b/Engine/Source/Runtime/Engine/Public/WorldPartition/WorldPartitionLevelStreamingDynamic.h
@@ -51,9 +51,6 @@ class UWorldPartitionLevelStreamingDynamic : public ULevelStreamingDynamic
 	virtual bool CanReplicateStreamingStatus() const override { return false; }
 
 #if WITH_EDITOR
-	static ENGINE_API UWorldPartitionLevelStreamingDynamic* LoadInEditor(UWorld* World, FName LevelStreamingName, const TArray<FWorldPartitionRuntimeCellObjectMapping>& InPackages);
-	static ENGINE_API void UnloadFromEditor(UWorldPartitionLevelStreamingDynamic* InLevelStreaming);
-
 	// Override ULevelStreaming
 	ENGINE_API virtual bool RequestLevel(UWorld* PersistentWorld, bool bAllowLevelLoadRequests, EReqLevelBlock BlockPolicy) override;
 	ENGINE_API virtual void BeginDestroy() override;
@@ -63,6 +60,12 @@ class UWorldPartitionLevelStreamingDynamic : public ULevelStreamingDynamic
 
 	void SetShouldPerformStandardLevelLoading(bool bInShouldPerformStandardLevelLoading) { bShouldPerformStandardLevelLoading = bInShouldPerformStandardLevelLoading; }
 
+	UE_DEPRECATED(5.7, "LoadInEditor() is no longer supported and will be removed in a future version.")
+	static ENGINE_API UWorldPartitionLevelStreamingDynamic* LoadInEditor(UWorld* World, FName LevelStreamingName, const TArray<FWorldPartitionRuntimeCellObjectMapping>& InPackages);
+
+	UE_DEPRECATED(5.7, "UnloadFromEditor() is no longer supported and will be removed in a future version.")
+	static ENGINE_API void UnloadFromEditor(UWorldPartitionLevelStreamingDynamic* InLevelStreaming);
+
 private:
 	ENGINE_API void CreateRuntimeLevel();
 	ENGINE_API bool IssueLoadRequests();
