World Partition - HLOD: Create the HLOD build world in a private transient package rather than the common transient package
* Prevents any possible interaction/conflicts when building HLODs for different worlds one after the other

diff --git a/Engine/Plugins/Editor/WorldPartitionHLODUtilities/Source/Private/WorldPartition/HLOD/Utilities/WorldPartitionHLODUtilities.cpp b/Engine/Plugins/Editor/WorldPartitionHLODUtilities/Source/Private/WorldPartition/HLOD/Utilities/WorldPartitionHLODUtilities.cpp
--- a/Engine/Plugins/Editor/WorldPartitionHLODUtilities/Source/Private/WorldPartition/HLOD/Utilities/WorldPartitionHLODUtilities.cpp
+++ b/Engine/Plugins/Editor/WorldPartitionHLODUtilities/Source/Private/WorldPartition/HLOD/Utilities/WorldPartitionHLODUtilities.cpp
@@ -795,7 +795,8 @@ uint32 FWorldPartitionHLODUtilities::BuildHLOD(AWorldPartitionHLOD* InHLODActor)
 
 	const bool bInformEngineOfWorld = false;
 	const bool bAddToRoot = true;
-	UWorld* TempWorld = UWorld::CreateWorld(EWorldType::Editor, bInformEngineOfWorld, MakeUniqueObjectName(GetTransientPackage(), UWorld::StaticClass(), TEXT("BuildHLODWorld")), GetTransientPackage(), bAddToRoot);
+	UPackage* TempPackage = NewObject<UPackage>(nullptr, *MakeUniqueObjectName(nullptr, UPackage::StaticClass(), TEXT("/Temp/BuildHLODPackage")).ToString(), RF_Transient);
+	UWorld* TempWorld = UWorld::CreateWorld(EWorldType::Editor, bInformEngineOfWorld, TEXT("BuildHLODWorld"), TempPackage, bAddToRoot);
 	check(TempWorld);
 
 	ON_SCOPE_EXIT
@@ -1044,7 +1045,8 @@ uint32 FWorldPartitionHLODUtilities::ComputeHLODHash(const AWorldPartitionHLOD*
 {
 	const bool bInformEngineOfWorld = false;
 	const bool bAddToRoot = true;
-	UWorld* TempWorld = UWorld::CreateWorld(EWorldType::Editor, bInformEngineOfWorld, MakeUniqueObjectName(GetTransientPackage(), UWorld::StaticClass(), TEXT("BuildHLODWorld")), GetTransientPackage(), bAddToRoot);
+	UPackage* TempPackage = NewObject<UPackage>(nullptr, *MakeUniqueObjectName(nullptr, UPackage::StaticClass(), TEXT("/Temp/BuildHLODPackage")).ToString(), RF_Transient);
+	UWorld* TempWorld = UWorld::CreateWorld(EWorldType::Editor, bInformEngineOfWorld, TEXT("BuildHLODWorld"), TempPackage, bAddToRoot);
 	check(TempWorld);
 
 	ON_SCOPE_EXIT
