World Partition - HLOD: Reduce FName bloating during HLOD generation
#Synced-CL 48765230

diff --git a/Engine/Plugins/Editor/WorldPartitionHLODUtilities/Source/Private/WorldPartition/HLOD/Utilities/WorldPartitionHLODUtilities.cpp b/Engine/Plugins/Editor/WorldPartitionHLODUtilities/Source/Private/WorldPartition/HLOD/Utilities/WorldPartitionHLODUtilities.cpp
--- a/Engine/Plugins/Editor/WorldPartitionHLODUtilities/Source/Private/WorldPartition/HLOD/Utilities/WorldPartitionHLODUtilities.cpp
+++ b/Engine/Plugins/Editor/WorldPartitionHLODUtilities/Source/Private/WorldPartition/HLOD/Utilities/WorldPartitionHLODUtilities.cpp
@@ -795,8 +795,7 @@ uint32 FWorldPartitionHLODUtilities::BuildHLOD(AWorldPartitionHLOD* InHLODActor)
 
 	const bool bInformEngineOfWorld = false;
 	const bool bAddToRoot = true;
-	UPackage* TempPackage = NewObject<UPackage>(nullptr, *MakeUniqueObjectName(nullptr, UPackage::StaticClass(), TEXT("/Temp/BuildHLODPackage")).ToString(), RF_Transient);
-	UWorld* TempWorld = UWorld::CreateWorld(EWorldType::Editor, bInformEngineOfWorld, TEXT("BuildHLODWorld"), TempPackage, bAddToRoot);
+	UWorld* TempWorld = UWorld::CreateWorld(EWorldType::Editor, bInformEngineOfWorld, MakeUniqueObjectName(GetTransientPackage(), UWorld::StaticClass(), TEXT("BuildHLODWorld")), GetTransientPackage(), bAddToRoot);
 	check(TempWorld);
 
 	ON_SCOPE_EXIT
diff --git a/Engine/Source/Runtime/Engine/Private/WorldPartition/HLOD/HLODSourceActorsFromCell.cpp b/Engine/Source/Runtime/Engine/Private/WorldPartition/HLOD/HLODSourceActorsFromCell.cpp
--- a/Engine/Source/Runtime/Engine/Private/WorldPartition/HLOD/HLODSourceActorsFromCell.cpp
+++ b/Engine/Source/Runtime/Engine/Private/WorldPartition/HLOD/HLODSourceActorsFromCell.cpp
@@ -32,6 +32,8 @@ bool UWorldPartitionHLODSourceActorsFromCell::LoadSourceActors(bool& bOutDirty,
 	UWorld* SourceWorld = HLODActor->GetWorld();
 
 	FLinkerInstancingContext InstancingContext;
+	InstancingContext.AddTag(ULevel::DontLoadExternalObjectsTag);
+	InstancingContext.AddTag(ULevel::DontLoadExternalFoldersTag);
 
 	// Don't do SoftObjectPath remapping for PersistentLevel actors because references can end up in different cells
 	InstancingContext.SetSoftObjectPathRemappingEnabled(false);
