World Partition - HLOD: New perceptual diff feature to avoid HLOD rebuilds when no significant visual changes are detected

* Added concept of HLOD Rebuild Policies
-> Policies can either approve, reject, or take no decision when evaluating whether a build is required
-> Policies store their own data (editor only) in the HLOD actor, which is then used as the basis in subsequent evaluations.
-> Multiple policies can be specified

* HLODRebuildPolicyHashCompare - Contains previous HLOD hash logic
-> Code migrated, logic mostly untouched
-> This HLOD rebuild policy computes a hash of various elements in order to detect any change that might impact the resulting HLODs. It is, in effect, extremely sensitive.
-> APPROVE a rebuild - if the old hash differs from the new hash
-> REJECT a rebuild - if the old hash equals the new hash

* HLODRebuildPolicyImageCompare - Fuzzy image comparison (using SSIM)
-> Takes screen captures of the SOURCE actors of an HLOD, over multiple angles (and for different GBuffer properties)
-> It will then perform an SSIM evaluation in order to assess if there is a significant visual change between the old and new data set.
-> APPROVE a rebuild - if the difference between old vs new images is greater than the configured threshold
-> REJECT a rebuild - if the difference between old vs new images is less than the configured threshold

Current default behavior is unchanged as HLODRebuildPolicyHashCompare is the only specified policy

#Synced-CL 48731319

diff --git a/Engine/Plugins/Editor/WorldPartitionHLODUtilities/Source/Private/WorldPartition/HLOD/RebuildPolicies/HLODRebuildPolicyHashCompare.cpp b/Engine/Plugins/Editor/WorldPartitionHLODUtilities/Source/Private/WorldPartition/HLOD/RebuildPolicies/HLODRebuildPolicyHashCompare.cpp
new file mode 100644
--- /dev/null
+++ b/Engine/Plugins/Editor/WorldPartitionHLODUtilities/Source/Private/WorldPartition/HLOD/RebuildPolicies/HLODRebuildPolicyHashCompare.cpp
@@ -0,0 +1,62 @@
+// Copyright Epic Games, Inc. All Rights Reserved.
+
+#include "WorldPartition/HLOD/RebuildPolicies/HLODRebuildPolicyHashCompare.h"
+#include "WorldPartition/HLOD/HLODActor.h"
+#include "WorldPartition/HLOD/HLODInstancedStaticMeshComponent.h"
+#include "WorldPartition/HLOD/HLODSourceActors.h"
+
+UHLODRebuildPolicyData* UHLODRebuildPolicyHashCompare::ComputeDataForRebuildPolicy(const AWorldPartitionHLOD* InHLODActor, const TArray<UActorComponent*>& InSourceComponents, const UHLODRebuildPolicyData* InExistingPolicyData) const
+{
+	TRACE_CPUPROFILER_EVENT_SCOPE(UHLODRebuildPolicyHashCompare::ComputeDataForRebuildPolicy);
+
+	FHLODHashBuilder HashBuilder;
+	ComputeHLODHash(HashBuilder, InHLODActor, InSourceComponents);
+
+	UHLODRebuildPolicyHashData* HashData = NewObject<UHLODRebuildPolicyHashData>();
+	HashData->HLODHash = HashBuilder.GetCrc();
+	HashData->HLODHashReport = HashBuilder.BuildHashReport();
+
+	return HashData;
+}
+
+EHLODRebuildPolicyDecision UHLODRebuildPolicyHashCompare::Evaluate(const AWorldPartitionHLOD* InHLODActor, const UHLODRebuildPolicyData* InOldData, const UHLODRebuildPolicyData* InNewData, FString& OutReason) const
+{
+	TRACE_CPUPROFILER_EVENT_SCOPE(UHLODRebuildPolicyHashCompare::Evaluate);
+
+	const UHLODRebuildPolicyHashData* OldData = Cast<UHLODRebuildPolicyHashData>(InOldData);
+	const UHLODRebuildPolicyHashData* NewData = CastChecked<UHLODRebuildPolicyHashData>(InNewData);
+
+	// Hash used to be stored directly in the HLOD actor, so retrieve it from there when necessary
+PRAGMA_DISABLE_DEPRECATION_WARNINGS
+	const uint32 OldHash = OldData ? OldData->HLODHash : InHLODActor->GetHLODHash();
+PRAGMA_ENABLE_DEPRECATION_WARNINGS
+	const uint32 NewHash = NewData->HLODHash;
+
+	return OldHash != NewHash ? EHLODRebuildPolicyDecision::ApproveRebuild : EHLODRebuildPolicyDecision::RejectRebuild;
+}
+
+void UHLODRebuildPolicyHashCompare::ComputeHLODHash(FHLODHashBuilder& InHashBuilder, const AWorldPartitionHLOD* InHLODActor, const TArray<UActorComponent*>& InSourceComponents) const
+{
+	// Base key, changing this will force a rebuild of all HLODs
+	FString HLODBaseKey = TEXT("0D33837AB1A04CC2AEEB04C5C217DE96");
+	InHashBuilder.HashField(HLODBaseKey, TEXT("HLODBaseKey"));
+
+	// HLOD Source Actors
+	if (ensure(InHLODActor->GetSourceActors()))
+	{
+		InHLODActor->GetSourceActors()->ComputeHLODHash(InHashBuilder);
+	}
+
+	// Min Visible Distance
+	InHashBuilder.HashField(InHLODActor->GetMinVisibleDistance(), TEXT("MinVisibleDistance"));
+
+	// ISM Component Class
+	TSubclassOf<UHLODInstancedStaticMeshComponent> HLODISMComponentClass = UHLODBuilder::GetInstancedStaticMeshComponentClass();
+	if (HLODISMComponentClass != UHLODInstancedStaticMeshComponent::StaticClass())
+	{
+		InHashBuilder.HashField(HLODISMComponentClass->GetPathName(), TEXT("HLODISMComponentClass"));
+	}
+
+	// Append all components hashes
+	UHLODBuilder::ComputeHLODHash(InHashBuilder, InSourceComponents);
+}
diff --git a/Engine/Plugins/Editor/WorldPartitionHLODUtilities/Source/Private/WorldPartition/HLOD/RebuildPolicies/HLODRebuildPolicyImageCompare.cpp b/Engine/Plugins/Editor/WorldPartitionHLODUtilities/Source/Private/WorldPartition/HLOD/RebuildPolicies/HLODRebuildPolicyImageCompare.cpp
new file mode 100644
--- /dev/null
+++ b/Engine/Plugins/Editor/WorldPartitionHLODUtilities/Source/Private/WorldPartition/HLOD/RebuildPolicies/HLODRebuildPolicyImageCompare.cpp
@@ -0,0 +1,1210 @@
+// Copyright Epic Games, Inc. All Rights Reserved.
+
+#include "WorldPartition/HLOD/RebuildPolicies/HLODRebuildPolicyImageCompare.h"
+#include "WorldPartition/HLOD/HLODActor.h"
+#include "WorldPartition/HLOD/HLODHashBuilder.h"
+#include "WorldPartition/HLOD/HLODLayer.h"
+#include "WorldPartition/HLOD/HLODSourceActors.h"
+
+#include "AssetCompilingManager.h"
+#include "Async/ParallelFor.h"
+#include "Components/PrimitiveComponent.h"
+#include "ContentStreaming.h"
+#include "EngineModule.h"
+#include "Engine/Canvas.h"
+#include "Engine/TextureRenderTarget2D.h"
+#include "Engine/World.h"
+#include "Framework/Application/SlateApplication.h"
+#include "HAL/FileManager.h"
+#include "ImageComparer.h"
+#include "IImageWrapper.h"
+#include "IImageWrapperModule.h"
+#include "Interfaces/IScreenShotComparisonModule.h"
+#include "Interfaces/IScreenShotManager.h"
+#include "Interfaces/IScreenShotToolsModule.h"
+#include "JsonObjectConverter.h"
+#include "LegacyScreenPercentageDriver.h"
+#include "MaterialShared.h"
+#include "Materials/MaterialInterface.h"
+#include "Materials/MaterialInstance.h"
+#include "Misc/FileHelper.h"
+#include "Misc/Paths.h"
+#include "Modules/ModuleManager.h"
+#include "ObjectCacheContext.h"
+#include "RenderingThread.h"
+#include "RHIGPUReadback.h"
+#include "SceneInterface.h"
+#include "SceneManagement.h"
+#include "TextureResource.h"
+#include "Widgets/SWindow.h"
+
+static bool GWorldPartitionHLODRebuildPolicyImageCompareDebugImageDiff = false;
+static FAutoConsoleVariableRef CVarWorldPartitionHLODRebuildPolicyImageCompareDebugImageDiff(
+	TEXT("wp.Editor.HLOD.RebuildPolicy.ImageCompare.DebugImageDiff"),
+	GWorldPartitionHLODRebuildPolicyImageCompareDebugImageDiff,
+	TEXT("Display comparison results in an editor window"),
+	ECVF_Scalability);
+
+struct FSceneCaptureConfig
+{
+	uint8					BytesPerPixel;
+	const TCHAR*			CaptureName;
+	EViewModeIndex			ViewModeIndex;
+	FName					VisualizationMode;
+};
+
+const FSceneCaptureConfig CaptureConfigs[EHLODCaptureImageType::Num] =
+{
+	{ 1,	TEXT("Alpha"),		VMI_Lit,				NAME_None },
+	{ 3,	TEXT("BaseColor"),	VMI_VisualizeBuffer,	TEXT("BaseColor") },
+	{ 3,	TEXT("Normal"),		VMI_VisualizeBuffer,	TEXT("WorldNormal") },
+	{ 3,	TEXT("Emissive"),	VMI_VisualizeBuffer,	TEXT("PreTonemapHDRColor") },
+	{ 1,	TEXT("Metallic"),	VMI_VisualizeBuffer,	TEXT("Metallic") },
+	{ 1,	TEXT("Roughness"),	VMI_VisualizeBuffer,	TEXT("Roughness") },
+	{ 1,	TEXT("Specular"),	VMI_VisualizeBuffer,	TEXT("Specular") },
+};
+
+const FRotator CaptureAngles[] =
+{
+	// Zenith (straight down)
+	FRotator(-90.f,   0.f,   0.f),
+
+	// Ring A @ elevation 40°
+	FRotator(-40.f,   0.f,   0.f),
+	FRotator(-40.f,  90.f,   0.f),
+	FRotator(-40.f, 180.f,   0.f),
+	FRotator(-40.f, 270.f,   0.f),
+
+	// Ring B @ elevation 10°
+	FRotator(-10.f,  45.f,   0.f),
+	FRotator(-10.f, 135.f,   0.f),
+	FRotator(-10.f, 225.f,   0.f),
+	FRotator(-10.f, 315.f,   0.f)
+};
+
+static FMatrix BuildOrthoMatrix(float InOrthoWidth, float InOrthoHeight)
+{
+	const double ZOffset = UE_FLOAT_HUGE_DISTANCE / 2.0;
+	return FReversedZOrthoMatrix(
+		InOrthoWidth / 2.0f,
+		InOrthoHeight / 2.0f,
+		0.5f / ZOffset,
+		ZOffset
+	);
+}
+
+static uint32 GetScreenSize(const FBoxSphereBounds& InBounds, double InExpectedDrawDistance)
+{
+	// Generate a projection matrix.
+	static const float ScreenX = 1920;
+	static const float ScreenY = 1080;
+	static const float HalfFOVRad = FMath::DegreesToRadians(45.0f);
+	static const FMatrix ProjectionMatrix = FPerspectiveMatrix(HalfFOVRad, ScreenX, ScreenY, 0.01);
+
+	float ScreenSizePercent = ComputeBoundsScreenSize(FVector::ZeroVector, (float)InBounds.SphereRadius, FVector(0.0, 0.0, InExpectedDrawDistance), ProjectionMatrix);
+	float ScreenSizePixel = ScreenSizePercent * ScreenX;
+	return (uint32)FMath::CeilToInt32(ScreenSizePixel);
+}
+
+static void PerformSceneWarmup(UWorld* InCaptureWorld, const TArray<UActorComponent*>& InSourceComponents)
+{
+	TRACE_CPUPROFILER_EVENT_SCOPE(Internal::PerformSceneWarmup);
+
+	// Wait for pending load requests.
+	FlushAsyncLoading();
+
+	// Inspired from UMaterialInterface::SubmitRemainingJobsForWorld() 
+	// Modified to handle MaterialInstance parent materials
+	{
+		TSet<UMaterialInterface*> MaterialsToCache;
+		FObjectCacheContextScope ObjectCacheScope;
+
+		for (UActorComponent* HLODRelevantComponent : InSourceComponents)
+		{
+			if (UPrimitiveComponent* PrimitiveComponent = Cast<UPrimitiveComponent>(HLODRelevantComponent))
+			{
+				TObjectCacheIterator<UMaterialInterface> UsedMaterials = ObjectCacheScope.GetContext().GetUsedMaterials(PrimitiveComponent->GetPrimitiveComponentInterface());
+				for (UMaterialInterface* MaterialInterface : UsedMaterials)
+				{
+					while (MaterialInterface && !MaterialsToCache.Contains(MaterialInterface))
+					{
+						if (!MaterialInterface->IsComplete())
+						{
+							MaterialsToCache.Add(MaterialInterface);
+						}
+
+						UMaterialInstance* MaterialInstance = Cast<UMaterialInstance>(MaterialInterface);
+						MaterialInterface = MaterialInstance ? MaterialInstance->Parent : nullptr;
+					};
+				}
+			}
+		}
+
+		if (MaterialsToCache.Num())
+		{
+			FMaterialUpdateContext UpdateContext(FMaterialUpdateContext::EOptions::SyncWithRenderingThread);
+			for (UMaterialInterface* Material : MaterialsToCache)
+			{
+				// This is needed because CacheShaders blindly recreates uniform buffers
+				// which can only be done if the draw command is going to be re-cached.
+				UpdateContext.AddMaterialInterface(Material);
+				Material->CacheShaders(EMaterialShaderPrecompileMode::Default);
+			}
+		}
+	}
+
+	// Wait for shader and other asset compilation to finish.
+	FAssetCompilingManager::Get().FinishAllCompilation();
+
+	// Force all mips to load.
+	UTexture::ForceUpdateTextureStreaming();
+
+	// Force all streamed resources to finish.
+	IStreamingManager::Get().StreamAllResources(0.0f);
+
+	// Ensure all deferred construction scripts are executed
+	FAssetCompilingManager::Get().ProcessAsyncTasks();
+
+	InCaptureWorld->SendAllEndOfFrameUpdates();
+}
+
+static void PerformSceneRender(FCanvas* InCanvas, FSceneViewFamily* InViewFamily)
+{
+	TRACE_CPUPROFILER_EVENT_SCOPE(Internal::PerformSceneRender);
+
+	GFrameCounter++;
+
+	ENQUEUE_RENDER_COMMAND(BeginFrameCommand)(
+		[CurrentFrameCounter = GFrameCounter](FRHICommandListImmediate& RHICmdList)
+		{
+			GFrameCounterRenderThread = CurrentFrameCounter;
+			GFrameNumberRenderThread++;
+			GPU_STATS_BEGINFRAME(RHICmdList);
+			FCoreDelegates::OnBeginFrameRT.Broadcast();
+		});
+
+	GetRendererModule().BeginRenderingViewFamily(InCanvas, InViewFamily);
+	
+	ENQUEUE_RENDER_COMMAND(EndFrameCommand)(
+		[](FRHICommandListImmediate& RHICmdList)
+		{
+			FCoreDelegates::OnEndFrameRT.Broadcast();
+			RHICmdList.EndFrame();
+			GPU_STATS_ENDFRAME(RHICmdList);
+		});
+}
+
+static void PerformViewWarmup(FCanvas* InCanvas, FSceneViewFamily* InViewFamily, uint32 InWarmupFrameCount)
+{
+	TRACE_CPUPROFILER_EVENT_SCOPE(Internal::PerformViewWarmup);
+
+	bool bCompilingAssets = false;
+	uint32 NumRender = 0;
+	
+	do
+	{
+		PerformSceneRender(InCanvas, InViewFamily);
+
+		NumRender++;
+
+		// Flush rendering commands, may queue assets compilation (shaders)
+		FlushRenderingCommands();
+
+		// If there are assets to be compiled, we need to wait for compilation to finish and start render over
+		bCompilingAssets = FAssetCompilingManager::Get().GetNumRemainingAssets() > 0;
+		if (bCompilingAssets)
+		{
+			FAssetCompilingManager::Get().FinishAllCompilation();
+
+			// Ensure all deferred construction scripts are executed
+			FAssetCompilingManager::Get().ProcessAsyncTasks();
+		}
+	} while (bCompilingAssets || NumRender < InWarmupFrameCount);
+}
+
+static TSharedPtr<FRHIGPUTextureReadback> RenderSceneToTexture(FSceneInterface* InScene, const bool bInPerformWarmup, uint32 InWarmupFrameCount, EHLODCaptureImageType InCaptureType, const FVector& InViewOrigin, const FRotator& InViewRotation, const FMatrix& InProjectionMatrix, const TSet<FPrimitiveComponentId>& InVisibleComponents, UTextureRenderTarget2D* InRenderTargetTexture)
+{
+	const FSceneCaptureConfig& CaptureConfig = CaptureConfigs[InCaptureType];
+
+	FEngineShowFlags ShowFlags(ESFIM_Game);
+	ApplyViewMode(CaptureConfig.ViewModeIndex, /*bPerspective=*/false, ShowFlags);
+	ShowFlags.DisableAdvancedFeatures();
+	ShowFlags.SetAmbientCubemap(false);
+	ShowFlags.SetAmbientOcclusion(false);
+	ShowFlags.SetAntiAliasing(true);
+	ShowFlags.SetTemporalAA(false);
+	ShowFlags.SetAtmosphere(false);
+	ShowFlags.SetBloom(false);
+	ShowFlags.SetCameraImperfections(false);
+	ShowFlags.SetCapsuleShadows(false);
+	ShowFlags.SetCloud(false);
+	ShowFlags.SetColorGrading(false);
+	ShowFlags.SetContactShadows(false);
+	ShowFlags.SetDepthOfField(false);
+	ShowFlags.SetDistanceFieldAO(false);
+	ShowFlags.SetDynamicShadows(false);
+	ShowFlags.SetEyeAdaptation(false);
+	ShowFlags.SetFog(false);
+	ShowFlags.SetGlobalIllumination(false);
+	ShowFlags.SetGrain(false);
+	ShowFlags.SetHighResScreenshotMask(false);
+	ShowFlags.SetHMDDistortion(false);
+	ShowFlags.SetIndirectLightingCache(false);
+	ShowFlags.SetLightShafts(false);
+	ShowFlags.SetLensFlares(false);
+	ShowFlags.SetMotionBlur(false);
+	ShowFlags.SetOnScreenDebug(false);
+	ShowFlags.SetPostProcessing(true);
+	ShowFlags.SetReflectionEnvironment(false);
+	ShowFlags.SetScreenPercentage(false);
+	ShowFlags.SetScreenSpaceReflections(false);
+	ShowFlags.SetSeparateTranslucency(false);
+	ShowFlags.SetSkyLighting(false);
+	ShowFlags.SetStereoRendering(false);
+	ShowFlags.SetTonemapper(false);
+	ShowFlags.SetToneCurve(false);
+	ShowFlags.SetVignette(false);
+	ShowFlags.SetVolumetricFog(false);
+	ShowFlags.SetVolumetricLightmap(false);	
+	
+	ShowFlags.SetVisualizeBuffer(CaptureConfig.ViewModeIndex == VMI_VisualizeBuffer);
+	ShowFlags.SetPostProcessMaterial(CaptureConfig.ViewModeIndex == VMI_VisualizeBuffer);
+
+	ShowFlags.SetLighting(InCaptureType == EHLODCaptureImageType::Capture_Emissive);
+	
+	// Use a constant (0) game time to enforce scene determinism between captures
+	static const FGameTime CaptureTime;
+
+	FTextureRenderTargetResource* RenderTargetResource = InRenderTargetTexture->GameThread_GetRenderTargetResource();
+
+	FSceneViewFamilyContext ViewFamily(
+		FSceneViewFamily::ConstructionValues(RenderTargetResource, InScene, ShowFlags)
+		.SetTime(CaptureTime)
+		.SetRealtimeUpdate(false)
+	);
+	ViewFamily.SetScreenPercentageInterface(new FLegacyScreenPercentageDriver(ViewFamily, 1.0f));
+
+	static const FMatrix ViewPlanesMatrix = FMatrix(
+		FPlane(0, 0, 1, 0),
+		FPlane(1, 0, 0, 0),
+		FPlane(0, 1, 0, 0),
+		FPlane(0, 0, 0, 1));
+
+	FSceneViewInitOptions ViewInitOptions;
+	ViewInitOptions.SetViewRectangle(FIntRect(0, 0, InRenderTargetTexture->SizeX, InRenderTargetTexture->SizeY));
+	ViewInitOptions.BackgroundColor = FLinearColor(FColor(0, 0, 0, 0));
+	ViewInitOptions.ViewFamily = &ViewFamily;
+	ViewInitOptions.ViewOrigin = InViewOrigin;
+	ViewInitOptions.ViewRotationMatrix = FInverseRotationMatrix(InViewRotation) * ViewPlanesMatrix;
+	ViewInitOptions.ProjectionMatrix = InProjectionMatrix;
+
+	if (InVisibleComponents.Num() > 0)
+	{
+		ViewInitOptions.ShowOnlyPrimitives = InVisibleComponents;
+	}
+
+	FSceneView* NewView = new FSceneView(ViewInitOptions);
+	NewView->CurrentBufferVisualizationMode = CaptureConfig.VisualizationMode;
+	ViewFamily.Views.Add(NewView);
+
+	FCanvas Canvas(RenderTargetResource, NULL, CaptureTime, InScene->GetFeatureLevel());
+
+	if (bInPerformWarmup)
+	{
+		PerformViewWarmup(&Canvas, &ViewFamily, InWarmupFrameCount);
+	}
+
+	PerformSceneRender(&Canvas, &ViewFamily);
+
+	TSharedPtr<FRHIGPUTextureReadback> ReadbackRequest = MakeShared<FRHIGPUTextureReadback>(TEXT("HLODSceneCaptureReadback"));
+
+	ENQUEUE_RENDER_COMMAND(HLODSceneCapture)([ReadbackRequest = ReadbackRequest, RenderTargetResource](FRHICommandListImmediate& RHICmdList) mutable
+	{
+		ReadbackRequest->EnqueueCopy(RHICmdList, RenderTargetResource->GetRenderTargetTexture());
+	});
+
+	return ReadbackRequest;
+}
+
+static void SavePNG(const TArray<FColor>& InPixels, int32 InImageWidth, int32 InImageHeight, const FString& InFilename)
+{
+	TRACE_CPUPROFILER_EVENT_SCOPE(SavePNG);
+
+	IImageWrapperModule& ImageWrapperModule = FModuleManager::LoadModuleChecked<IImageWrapperModule>(FName("ImageWrapper"));
+	TSharedPtr<IImageWrapper> ImageWrapper = ImageWrapperModule.CreateImageWrapper(EImageFormat::PNG);
+
+	if (ImageWrapper.IsValid() && ImageWrapper->SetRaw(InPixels.GetData(), InPixels.NumBytes(), InImageWidth, InImageHeight, ERGBFormat::BGRA, 8))
+	{
+		TArray64<uint8> RawData = ImageWrapper->GetCompressed();
+		FFileHelper::SaveArrayToFile(RawData, *InFilename);
+	}
+}
+
+
+static bool AreBoxesWithinPercent(const FBox& InCandidate, const FBox& InReference, double InPercent)
+{
+	const double Alpha = FMath::Clamp(InPercent, 0.f, 1.f);
+
+	FVector RefSize = InReference.GetSize();
+	RefSize.X = FMath::Max(RefSize.X, KINDA_SMALL_NUMBER);
+	RefSize.Y = FMath::Max(RefSize.Y, KINDA_SMALL_NUMBER);
+	RefSize.Z = FMath::Max(RefSize.Z, KINDA_SMALL_NUMBER);
+
+	const double WorstMin = ((InCandidate.Min - InReference.Min).GetAbs() / RefSize).GetMax();
+	const double WorstMax = ((InCandidate.Max - InReference.Max).GetAbs() / RefSize).GetMax();
+
+	return FMath::Max(WorstMin, WorstMax) <= Alpha;
+}
+
+
+UHLODRebuildPolicyImageCompare::UHLODRebuildPolicyImageCompare(const FObjectInitializer& ObjectInitializer)
+	: Super(ObjectInitializer)
+{
+}
+
+UHLODRebuildPolicyData* UHLODRebuildPolicyImageCompare::ComputeDataForRebuildPolicy(const AWorldPartitionHLOD* InHLODActor, const TArray<UActorComponent*>& InSourceComponents, const UHLODRebuildPolicyData* InExistingPolicyData) const
+{
+	TRACE_CPUPROFILER_EVENT_SCOPE(UHLODRebuildPolicyImageCompare::ComputeDataForRebuildPolicy);
+
+	if (InSourceComponents.IsEmpty())
+	{
+		return nullptr;
+	}
+
+	// Propagate alpha to render texture during scene capture
+	// Used to obtain an alpha mask of what is actually rendered to the scene
+	static IConsoleVariable* CVarPropagateAlpha = IConsoleManager::Get().FindConsoleVariable(TEXT("r.PostProcessing.PropagateAlpha"));
+	CVarPropagateAlpha->Set(true, EConsoleVariableFlags(ECVF_SetByCode | ECVF_Set_NoSinkCall_Unsafe));
+	ON_SCOPE_EXIT
+	{
+		CVarPropagateAlpha->Unset(EConsoleVariableFlags(ECVF_SetByCode | ECVF_Set_NoSinkCall_Unsafe));
+	};
+
+	UWorld* CaptureWorld = InSourceComponents[0]->GetWorld();
+	
+	PerformSceneWarmup(CaptureWorld, InSourceComponents);
+
+	FBoxSphereBounds::Builder BoundsBuilder;
+	TSet<FPrimitiveComponentId> VisibleComponents;
+	for (UActorComponent* HLODRelevantComponent : InSourceComponents)
+	{
+		if (USceneComponent* SceneComponent = Cast<USceneComponent>(HLODRelevantComponent))
+		{
+			if (UPrimitiveComponent* PrimitiveComponent = Cast<UPrimitiveComponent>(SceneComponent))
+			{
+				BoundsBuilder += PrimitiveComponent->GetBounds().GetBox();
+				VisibleComponents.Add(PrimitiveComponent->GetPrimitiveSceneId());
+			}
+		}
+	}
+	FBox TargetBounds = FBoxSphereBounds(BoundsBuilder).GetBox();
+
+	const UHLODRebuildPolicyImageData* ExistingPolicyData = Cast<UHLODRebuildPolicyImageData>(InExistingPolicyData);
+	if (ExistingPolicyData)
+	{
+		if (AreBoxesWithinPercent(TargetBounds, ExistingPolicyData->CaptureBounds, CaptureImageMargin))
+		{
+			TargetBounds = ExistingPolicyData->CaptureBounds;
+		}
+	}
+
+	FVector CameraTarget = TargetBounds.GetCenter();
+	double CaptureDistance = TargetBounds.GetExtent().GetMin() + InHLODActor->GetMinVisibleDistance();
+
+	UHLODRebuildPolicyImageData* ImageData = NewObject<UHLODRebuildPolicyImageData>();
+	ImageData->CaptureBounds = TargetBounds;
+	ImageData->CaptureFrames.Reserve(UE_ARRAY_COUNT(CaptureAngles));
+
+	const uint32 MinImageSize = 16;
+	uint32 ImageSize = FMath::Clamp(GetScreenSize(TargetBounds, CaptureDistance), MinImageSize, MaxImageSize);
+	
+	// Flag kept to ensure we perform warmup of the scene only once
+	bool bPerformedWarmup = false;
+
+	// Struct used to enqueue and process readback operations
+	struct FPendingReadback
+	{
+		uint32 AngleIndex;
+		uint32 ImageIndex;
+		UTextureRenderTarget2D* RenderTargetTexture;
+		TSharedPtr<FRHIGPUTextureReadback> TextureReadback;
+		TArray<FColor>& Pixels;
+	};
+
+	// Perform scene capture from multiple angles
+	for (uint32 AngleIndex = 0; AngleIndex < UE_ARRAY_COUNT(CaptureAngles); ++AngleIndex)
+	{
+		const FRotator& CaptureAngle = CaptureAngles[AngleIndex];
+	
+		FHLODCaptureFrame& CaptureFrame = ImageData->CaptureFrames.Emplace_GetRef();
+
+		CaptureFrame.CameraLocation = CameraTarget + CaptureAngle.Vector() * CaptureDistance;
+		CaptureFrame.CameraRotation = CaptureAngle;
+		CaptureFrame.OrthoWidth = (float)TargetBounds.GetSize().Length() * (1 + CaptureImageMargin);
+
+		// For each camera angle, capture different GBuffer properties (base color, normal, metallic, roughnes, etc.)
+		for (uint32 CaptureIndex = 0; CaptureIndex < EHLODCaptureImageType::Num; ++CaptureIndex)
+		{
+			EHLODCaptureImageType CaptureType = (EHLODCaptureImageType)CaptureIndex;
+
+			FMatrix ProjectionMatrix = BuildOrthoMatrix(CaptureFrame.OrthoWidth, CaptureFrame.OrthoWidth);
+
+			UTextureRenderTarget2D* RenderTargetTexture = NewObject<UTextureRenderTarget2D>();
+			RenderTargetTexture->InitCustomFormat(ImageSize, ImageSize, PF_B8G8R8A8, false);
+
+			TSharedPtr<FRHIGPUTextureReadback> TextureReadback = RenderSceneToTexture(
+				CaptureWorld->Scene,
+				!bPerformedWarmup,
+				WarmupFrameCount,
+				CaptureType,
+				CaptureFrame.CameraLocation,
+				CaptureFrame.CameraRotation,
+				ProjectionMatrix,
+				VisibleComponents,
+				RenderTargetTexture);
+
+			FPendingReadback PendingReadback = { AngleIndex, CaptureType, RenderTargetTexture, TextureReadback, CaptureFrame.Images[CaptureType].Pixels };
+			ENQUEUE_RENDER_COMMAND(PerformReadback)(
+				[PendingReadback, ImageSize](FRHICommandListImmediate& RHICmdList) mutable
+				{
+					TRACE_CPUPROFILER_EVENT_SCOPE(DrainAllReadbacks);
+
+					// Wait until the GPU copy finished.
+					while (!PendingReadback.TextureReadback->IsReady())
+					{
+						RHICmdList.ImmediateFlush(EImmediateFlushType::FlushRHIThread);
+						FPlatformProcess::SleepNoStats(0.0f);
+					}
+
+					int32 RowPitchInPixels = 0;
+					const void* Src = PendingReadback.TextureReadback->Lock(RowPitchInPixels);
+					const int32 RowPitchInBytes = RowPitchInPixels * sizeof(FColor);
+					const uint32 DstBytesPerRow = ImageSize * sizeof(FColor);
+
+					PendingReadback.Pixels.SetNumUninitialized(ImageSize * ImageSize);
+					uint8* Dst = reinterpret_cast<uint8*>(PendingReadback.Pixels.GetData());
+
+					for (uint32 y = 0; y < ImageSize; ++y)
+					{
+						FMemory::Memcpy(Dst + y * DstBytesPerRow, static_cast<const uint8*>(Src) + y * RowPitchInBytes, DstBytesPerRow);
+					}
+
+					PendingReadback.TextureReadback->Unlock();
+				});
+		}
+
+		bPerformedWarmup = true;
+	}
+
+	FlushRenderingCommands();
+
+	// Finalize captured data
+	// Alpha mask -> Fill RGB as 255,255,255. Captured alpha is actually the inverse, so we need to reverse it.
+	// Others -> Copy alpha over from alpha mask
+	ParallelFor(ImageData->CaptureFrames.Num(), [&](int32 iCaptureFrame)
+	{
+		FHLODCaptureFrame& CaptureFrame = ImageData->CaptureFrames[iCaptureFrame];
+
+		for (int32 ImageIndex = 0; ImageIndex < EHLODCaptureImageType::Num; ++ImageIndex)
+		{
+			const FSceneCaptureConfig& CaptureConfig = CaptureConfigs[ImageIndex];
+
+			FHLODCaptureImage& Image = CaptureFrame.Images[ImageIndex];
+			Image.Width = ImageSize;
+			Image.Height = ImageSize;
+			Image.BytesPerPixel = CaptureConfig.BytesPerPixel;
+
+			TArray<FColor>& Pixels = Image.Pixels;
+			TArray<FColor>& PixelsAlpha = CaptureFrame.Images[Capture_AlphaMask].Pixels;
+
+			if (ImageIndex == Capture_AlphaMask)
+			{
+				for (int32 iPixel = 0; iPixel < Pixels.Num(); ++iPixel)
+				{
+					Pixels[iPixel].A = 255 - Pixels[iPixel].A;
+					Pixels[iPixel].R = 255;
+					Pixels[iPixel].G = 255;
+					Pixels[iPixel].B = 255;
+				}
+			}
+			else
+			{
+				for (int32 iPixel = 0; iPixel < Pixels.Num(); ++iPixel)
+				{
+					Pixels[iPixel].A = PixelsAlpha[iPixel].A;
+				}
+			}
+		}
+	});
+
+	FHLODHashBuilder HashBuilder;
+	const UWorldPartitionHLODSourceActors* HLODSourceActors = InHLODActor->GetSourceActors();
+	const UHLODLayer* HLODLayer = HLODSourceActors ? HLODSourceActors->GetHLODLayer() : nullptr;
+	if (HLODLayer)
+	{
+		HLODLayer->ComputeHLODHash(HashBuilder);
+		ImageData->HLODLayerHash = HashBuilder.GetCrc();
+	}
+	
+	ImageData->bNeedsDataCompression = true;
+	ImageData->bNeedsDataDecompression = false;
+
+	return ImageData;
+}
+
+template<bool bSingleChannel, typename SetErrorFunc>
+void CompareStructuralSimilarity(int32 InImageWidth, int32 InImageHeight, const TArray<FColor>& InPixelsA, const TArray<FColor>& InPixelsB, SetErrorFunc&& InSetErrorFunc, float& OutMaxLocalDifference, float& OutGlobalDifference)
+{
+	// Implementation of https://en.wikipedia.org/wiki/Structural_similarity
+
+	static const int32 BitsPerComponent = 8;
+	static const double K1 = 0.01;
+	static const double K2 = 0.03;
+	static const double L = (1 << BitsPerComponent) - 1;
+	static const double C1 = FMath::Pow(K1 * L, 2);
+	static const double C2 = FMath::Pow(K2 * L, 2);
+
+	double WeightedSSIM = 0.0;
+	double MaxWindowDSSIM = 0;
+	
+	// Global alpha-weighted total
+	double WeightedTotalAlpha = 0.0;
+	
+	// Gaussian kernel (11x11, sigma=1.5)
+	// Static, computed once
+	constexpr int32 KernelRadius = 5;
+	constexpr int32 KernelSize = KernelRadius * 2 + 1;
+	static const TStaticArray<double, KernelSize> Kernel = [] 
+	{
+		TStaticArray<double, KernelSize> K{};
+		const double Sigma = 1.5;
+		const double TwoSigma2 = 2.0 * Sigma * Sigma;
+		double Sum = 0.0;
+		for (int i = -KernelRadius; i <= KernelRadius; ++i)
+		{
+			const double Value = FMath::Exp(-(double)(i * i) / TwoSigma2);
+			K[i + KernelRadius] = Value;
+			Sum += Value;
+		}
+		for (int i = 0; i < KernelSize; ++i)
+		{
+			K[i] /= Sum;
+		}
+		return K;
+	}();
+
+	// Index directly in memory rather than going through the TArray operator[], as
+	// it adds quite a bit of overhead in repeated accesses like the kernel processing below.
+	const FColor* PixelsA = InPixelsA.GetData();
+	const FColor* PixelsB = InPixelsB.GetData();
+
+	for (int32 X = 0; X < InImageWidth; ++X)
+	{
+		for (int32 Y = 0; Y < InImageHeight; ++Y)
+		{
+			// Per-window accumulators
+			double SumW = 0.0;
+			double SumA_R = 0.0, SumB_R = 0.0, SumAA_R = 0.0, SumBB_R = 0.0, SumAB_R = 0.0;
+			double SumA_G = 0.0, SumB_G = 0.0, SumAA_G = 0.0, SumBB_G = 0.0, SumAB_G = 0.0;
+			double SumA_B = 0.0, SumB_B = 0.0, SumAA_B = 0.0, SumBB_B = 0.0, SumAB_B = 0.0;
+
+			for (int dy = -KernelRadius; dy <= KernelRadius; ++dy)
+			{
+				const int32 yy = FMath::Clamp(Y + dy, 0, InImageHeight - 1);
+				const double WeightY = Kernel[dy + KernelRadius];
+
+				for (int dx = -KernelRadius; dx <= KernelRadius; ++dx)
+				{
+					const int32 xx = FMath::Clamp(X + dx, 0, InImageWidth - 1);
+					const double Weight = WeightY * Kernel[dx + KernelRadius];
+
+					// Per-channel alpha-weighted accumulators (combine kernel with alpha)
+					const FColor& CA = PixelsA[xx + yy * InImageWidth];
+					const FColor& CB = PixelsB[xx + yy * InImageWidth];
+					const double AlphaA = CA.A / 255.0;
+					const double AlphaB = CB.A / 255.0;
+					const double WeightedAlpha = (AlphaA + AlphaB) - (AlphaA * AlphaB);
+					if (WeightedAlpha > 0.0)
+					{
+						const double WA = Weight * WeightedAlpha;
+						SumW += WA;
+
+						const double Ar = (double)CA.R;
+						const double Br = (double)CB.R;
+						SumA_R += WA * Ar;  SumB_R += WA * Br;  SumAA_R += WA * Ar * Ar;  SumBB_R += WA * Br * Br;  SumAB_R += WA * Ar * Br;
+
+						if (!bSingleChannel)
+						{
+							const double Ag = (double)CA.G;
+							const double Bg = (double)CB.G;
+							SumA_G += WA * Ag;  SumB_G += WA * Bg;  SumAA_G += WA * Ag * Ag;  SumBB_G += WA * Bg * Bg;  SumAB_G += WA * Ag * Bg;
+
+							const double Ab = (double)CA.B;
+							const double Bb = (double)CB.B;
+							SumA_B += WA * Ab;  SumB_B += WA * Bb;  SumAA_B += WA * Ab * Ab;  SumBB_B += WA * Bb * Bb;  SumAB_B += WA * Ab * Bb;
+						}
+					}
+				}
+			}
+
+			double WindowSSIM = 1.0;
+			double WindowDSSIM = 0.0;
+			double WindowAlphaWeight = 0.0;
+
+			// Skip if entire window is effectively transparent, in that case treat as perfect match with zero weight.
+			if (SumW > 0.0)
+			{
+				auto ChannelSSIM = [&](double SumA, double SumB, double SumAA, double SumBB, double SumAB) -> double
+				{
+					const double AverageAChannel = SumA / SumW;
+					const double AverageBChannel = SumB / SumW;
+
+					double VarianceAChannel = SumAA / SumW - AverageAChannel * AverageAChannel;
+					double VarianceBChannel = SumBB / SumW - AverageBChannel * AverageBChannel;
+					double CovarianceABChannel = SumAB / SumW - AverageAChannel * AverageBChannel;
+
+					VarianceAChannel = FMath::Max(0.0, VarianceAChannel);
+					VarianceBChannel = FMath::Max(0.0, VarianceBChannel);
+
+					const double LuminanceChannel = (2.0 * AverageAChannel * AverageBChannel + C1) / (AverageAChannel * AverageAChannel + AverageBChannel * AverageBChannel + C1);
+					const double ContrastChannel = (2.0 * CovarianceABChannel + C2) / (VarianceAChannel + VarianceBChannel + C2);
+
+					return LuminanceChannel * ContrastChannel;
+				};
+
+				const double SSIM_R = ChannelSSIM(SumA_R, SumB_R, SumAA_R, SumBB_R, SumAB_R);
+
+				if (!bSingleChannel)
+				{
+					const double SSIM_G = ChannelSSIM(SumA_G, SumB_G, SumAA_G, SumBB_G, SumAB_G);
+					const double SSIM_B = ChannelSSIM(SumA_B, SumB_B, SumAA_B, SumBB_B, SumAB_B);
+					WindowSSIM = (SSIM_R + SSIM_G + SSIM_B) / 3.0;
+				}
+				else
+				{
+					WindowSSIM = SSIM_R;
+				}
+
+				// center pixel's weight
+				const FColor& CA = PixelsA[X + Y * InImageWidth];
+				const FColor& CB = PixelsB[X + Y * InImageWidth];
+				const double AlphaA = CA.A / 255.0;
+				const double AlphaB = CB.A / 255.0;
+				const double WeightedAlpha = (AlphaA + AlphaB) - (AlphaA * AlphaB);
+				WindowAlphaWeight = WeightedAlpha;
+			}
+
+			WindowSSIM = FMath::Clamp(WindowSSIM, -1.0, 1.0);
+			WindowDSSIM = (1 - WindowSSIM) / 2;
+
+			InSetErrorFunc(X, Y, WindowDSSIM);
+		
+			MaxWindowDSSIM = FMath::Max(MaxWindowDSSIM, WindowDSSIM);
+
+			// Use alpha-weighted contribution
+			WeightedSSIM += WindowSSIM * WindowAlphaWeight;
+			WeightedTotalAlpha += WindowAlphaWeight;
+		}
+	}
+
+	double SSIM = (WeightedTotalAlpha > 0.0) ? (WeightedSSIM / WeightedTotalAlpha) : 1.0;
+	double DSSIM = (1.0 - SSIM) / 2.0;
+
+	OutMaxLocalDifference = (float)FMath::Clamp(MaxWindowDSSIM, 0.0, 1.0);
+	OutGlobalDifference = (float)FMath::Clamp(DSSIM, 0.0, 1.0);
+}
+
+struct FComparisonResult
+{
+	const FHLODCaptureImage*	ImageA;
+	const FHLODCaptureImage*	ImageB;
+	EHLODCaptureImageType		CaptureType;
+	FString						ComparisonName;
+
+	TArray<FColor> ErrorResult;
+	float MaxLocalDSSIM = 0.0;
+	float DSSIM = 0.0;
+};
+
+static FComparisonResult Compare(const FHLODCaptureFrame& InFrameA, const FHLODCaptureFrame& InFrameB, uint32 InCaptureFrameIdx, EHLODCaptureImageType InImageType, bool bInGatherErrorImage)
+{
+	TRACE_CPUPROFILER_EVENT_SCOPE(CompareStructuralSimilarity);
+
+	FComparisonResult ComparisonResult;
+	ComparisonResult.ImageA = &InFrameA.Images[InImageType];
+	ComparisonResult.ImageB = &InFrameB.Images[InImageType];
+	ComparisonResult.CaptureType = InImageType;
+	ComparisonResult.ComparisonName = FString::Printf(TEXT("Frame%d_%s"), InCaptureFrameIdx, CaptureConfigs[InImageType].CaptureName);
+
+	check(ComparisonResult.ImageA->Width == ComparisonResult.ImageB->Width);
+	check(ComparisonResult.ImageA->Height == ComparisonResult.ImageB->Height);
+
+	if (bInGatherErrorImage)
+	{
+		ComparisonResult.ErrorResult.SetNumUninitialized(ComparisonResult.ImageA->Width * ComparisonResult.ImageA->Height);
+	}
+
+	auto SetError = [&](int32 x, int32 y, double WindowDSSIM)
+	{
+		if (bInGatherErrorImage)
+		{
+			uint8 QuantizedError = FColor::QuantizeUNormFloatTo8((float)WindowDSSIM);
+			FColor ErrorColor(QuantizedError, QuantizedError, QuantizedError, 255);
+			ComparisonResult.ErrorResult[y * ComparisonResult.ImageA->Width + x] = ErrorColor;
+		}
+	};
+
+	const bool bSingleChannel = ComparisonResult.ImageA->BytesPerPixel == 1;
+	if (bSingleChannel)
+	{
+		CompareStructuralSimilarity<true>(ComparisonResult.ImageA->Width, ComparisonResult.ImageA->Height, ComparisonResult.ImageA->Pixels, ComparisonResult.ImageB->Pixels, SetError, ComparisonResult.MaxLocalDSSIM, ComparisonResult.DSSIM);
+	}
+	else
+	{
+		CompareStructuralSimilarity<false>(ComparisonResult.ImageA->Width, ComparisonResult.ImageA->Height, ComparisonResult.ImageA->Pixels, ComparisonResult.ImageB->Pixels, SetError, ComparisonResult.MaxLocalDSSIM, ComparisonResult.DSSIM);
+	}
+
+	return ComparisonResult;
+}
+
+static void ExecuteDebugComparison(const TArray<FComparisonResult>& InComparisonResults, float MaximumLocalErrorTolerance, float MaximumGlobalErrorTolerance)
+{
+	const FString RunLabel = TEXT("HLOD_") + FDateTime::Now().ToString(TEXT("%Y-%m-%d_%H-%M-%S"));
+	const FString ReportsRoot = FPaths::AutomationReportsDir();
+	const FString ImportRoot = FPaths::Combine(ReportsRoot, RunLabel);
+	const FString SubjectName = TEXT("HLOD");
+
+	auto MakePathRelative = [&](const FString& Abs)
+	{
+		FString Rel = Abs;
+		FPaths::MakePathRelativeTo(Rel, *ImportRoot);
+		return Rel;
+	};
+
+	for (const FComparisonResult& ComparisonResult : InComparisonResults)
+	{
+		const FString PairName = ComparisonResult.ComparisonName;
+		const FString ReportDir = FPaths::Combine(ImportRoot, PairName);
+		const FString OldImagePath = FPaths::Combine(ReportDir, TEXT("Approved.png"));
+		const FString NewImagePath = FPaths::Combine(ReportDir, TEXT("Incoming.png"));
+		const FString DeltaPath = FPaths::Combine(ReportDir, TEXT("Delta.png"));
+		const FString JsonPath = FPaths::Combine(ReportDir, TEXT("Report.json"));
+
+		IFileManager::Get().MakeDirectory(*ReportDir, /*Tree=*/true);
+
+		SavePNG(ComparisonResult.ImageA->Pixels, ComparisonResult.ImageA->Width, ComparisonResult.ImageA->Height, OldImagePath);
+		SavePNG(ComparisonResult.ImageB->Pixels, ComparisonResult.ImageB->Width, ComparisonResult.ImageB->Height, NewImagePath);
+		SavePNG(ComparisonResult.ErrorResult, ComparisonResult.ImageA->Width, ComparisonResult.ImageA->Height, DeltaPath);
+
+		FImageComparisonResult Result;
+		Result.ApprovedFilePath = OldImagePath;
+		Result.ReportApprovedFilePath = MakePathRelative(OldImagePath);
+		Result.IncomingFilePath = NewImagePath;
+		Result.ReportIncomingFilePath = MakePathRelative(NewImagePath);
+		Result.ComparisonFilePath = DeltaPath;
+		Result.ReportComparisonFilePath = MakePathRelative(DeltaPath);
+		
+		Result.ScreenshotPath = PairName;
+		Result.IdealApprovedFolderPath = FPaths::GetPath(OldImagePath);
+		Result.CreationTime = FDateTime::Now();
+		Result.bSkipAttachingImages = false;
+
+		Result.MaxLocalDifference = ComparisonResult.MaxLocalDSSIM;
+		Result.GlobalDifference = ComparisonResult.DSSIM;
+
+		FImageTolerance ImageTolerance;
+		ImageTolerance.IgnoreAntiAliasing = true;
+		ImageTolerance.IgnoreColors = true;
+		ImageTolerance.MaximumGlobalError = MaximumGlobalErrorTolerance;
+		ImageTolerance.MaximumLocalError = MaximumLocalErrorTolerance;
+		Result.Tolerance = ImageTolerance;
+
+		FString Json;
+		if (FJsonObjectConverter::UStructToJsonObjectString(Result, Json))
+		{
+			FFileHelper::SaveStringToFile(Json, *JsonPath, FFileHelper::EEncodingOptions::ForceUTF8WithoutBOM);
+		}
+	}
+
+	IScreenShotToolsModule& Tools = FModuleManager::LoadModuleChecked<IScreenShotToolsModule>("ScreenShotComparisonTools");
+	IScreenShotManagerPtr Manager = Tools.GetScreenShotManager();
+	Manager->OpenComparisonReportsAsync(ImportRoot);
+
+	IScreenShotComparisonModule& ComparisonUI = FModuleManager::LoadModuleChecked<IScreenShotComparisonModule>("ScreenShotComparison");
+
+	TSharedRef<SWidget> Browser = ComparisonUI.CreateScreenShotComparison(Manager.ToSharedRef());
+
+	TSharedRef<SWindow> Win = SNew(SWindow)
+		.Title(FText::FromString(TEXT("Image Comparison")))
+		.ClientSize(FVector2D(1200, 800))
+		[Browser];
+
+	FSlateApplication::Get().AddWindow(Win);
+}
+
+
+EHLODRebuildPolicyDecision UHLODRebuildPolicyImageCompare::Evaluate(const AWorldPartitionHLOD* InHLODActor, const UHLODRebuildPolicyData* InOldData, const UHLODRebuildPolicyData* InNewData, FString& OutReason) const
+{
+	TRACE_CPUPROFILER_EVENT_SCOPE(UHLODRebuildPolicyImageCompare::Evaluate);
+
+	const UHLODRebuildPolicyImageData* OldData = Cast<UHLODRebuildPolicyImageData>(InOldData);
+	const UHLODRebuildPolicyImageData* NewData = CastChecked<UHLODRebuildPolicyImageData>(InNewData);
+
+	if (OldData == nullptr)
+	{
+		OutReason = TEXT("no old data to compare with");
+		return EHLODRebuildPolicyDecision::None;
+	}
+
+	// == Special handling for the HLOD layer settings ==
+	// As this policy is only comparing screen captures of the INPUT actors, it would not take
+	// into account changes to the HLOD layer settings, as these can affect the OUTPUT of the
+	// HLOD build.
+	// The HLOD build should be allowed to be performed in such case.
+	if (OldData->HLODLayerHash != NewData->HLODLayerHash)
+	{
+		OutReason = TEXT("hlod layer settings were modified");
+		return EHLODRebuildPolicyDecision::ApproveRebuild;
+	}
+
+	if (OldData->CaptureFrames.Num() != NewData->CaptureFrames.Num())
+	{
+		OutReason = FString::Printf(TEXT("number of frames to compare differ, %d vs %d"), OldData->CaptureFrames.Num(), NewData->CaptureFrames.Num());
+		return EHLODRebuildPolicyDecision::None;
+	}
+
+	for (int32 CaptureFrameIdx = 0; CaptureFrameIdx < OldData->CaptureFrames.Num(); ++CaptureFrameIdx)
+	{
+		const FHLODCaptureFrame& FrameA = OldData->CaptureFrames[CaptureFrameIdx];
+		const FHLODCaptureFrame& FrameB = NewData->CaptureFrames[CaptureFrameIdx];
+		
+		for (int32 ImageIdx = 0; ImageIdx < EHLODCaptureImageType::Num; ++ImageIdx)
+		{
+			const FHLODCaptureImage& ImageA = FrameA.Images[ImageIdx];
+			const FHLODCaptureImage& ImageB = FrameB.Images[ImageIdx];
+
+			if ((ImageA.Width != ImageB.Width) || (ImageA.Height != ImageB.Height))
+			{
+				OutReason = FString::Printf(TEXT("image sizes differ, %dx%d vs %dx%d"), ImageA.Width, ImageA.Height, ImageB.Width, ImageB.Height);
+				return EHLODRebuildPolicyDecision::None;
+			}
+
+			if (ImageA.BytesPerPixel != ImageB.BytesPerPixel)
+			{
+				OutReason = FString::Printf(TEXT("image bpp differ, %d vs %d"), ImageA.BytesPerPixel, ImageB.BytesPerPixel);
+				return EHLODRebuildPolicyDecision::None;
+			}
+		}
+	}
+
+	// Ensure stored images are available in uncompressed state
+	bool bDataValid = true;
+	bDataValid &= OldData->DecompressData();
+	bDataValid &= NewData->DecompressData();
+	if (!bDataValid)
+	{
+		OutReason = TEXT("failed to decompress data");
+		return EHLODRebuildPolicyDecision::None;
+	}
+
+	const int32 NumFrames = OldData->CaptureFrames.Num();
+	const int32 NumComparisonPerFrame = (EHLODCaptureImageType::Num - 1); // -1 as we do not compare the alpha mask directly
+
+	TArray<FComparisonResult> ComparisonResults;
+	ComparisonResults.SetNum(NumFrames * NumComparisonPerFrame);
+
+	ParallelFor(ComparisonResults.Num(), [&](int32 Idx)
+	{
+		int32 CaptureFrameIdx = Idx / NumComparisonPerFrame;
+		const FHLODCaptureFrame& FrameA = OldData->CaptureFrames[CaptureFrameIdx];
+		const FHLODCaptureFrame& FrameB = NewData->CaptureFrames[CaptureFrameIdx];
+
+		EHLODCaptureImageType ImageIdx = (EHLODCaptureImageType)((Idx % NumComparisonPerFrame) + 1); // +1 as we do not compare the alpha mask directly
+		
+		ComparisonResults[Idx] = Compare(FrameA, FrameB, CaptureFrameIdx, ImageIdx, GWorldPartitionHLODRebuildPolicyImageCompareDebugImageDiff);
+	});
+
+	// If enabled, open up a comparison window
+	if (GWorldPartitionHLODRebuildPolicyImageCompareDebugImageDiff)
+	{
+		ExecuteDebugComparison(ComparisonResults, MaximumLocalError, MaximumGlobalError);
+	}
+
+	double MaxLocalDiff = 0.0f;
+	double MaxGlobalDiff = 0.0f;
+	for (const FComparisonResult& ComparisonResult : ComparisonResults)
+	{
+		MaxLocalDiff = FMath::Max(MaxLocalDiff, ComparisonResult.MaxLocalDSSIM);
+		MaxGlobalDiff = FMath::Max(MaxGlobalDiff, ComparisonResult.DSSIM);
+	}
+
+	if (MaxGlobalDiff > MaximumGlobalError)
+	{
+		OutReason = FString::Printf(TEXT("global diff: %.2f%% > thres: %.2f%%"), MaxGlobalDiff * 100, MaximumGlobalError * 100);
+		return EHLODRebuildPolicyDecision::ApproveRebuild;
+	}
+
+	if (MaxLocalDiff > MaximumLocalError)
+	{
+		OutReason = FString::Printf(TEXT("local diff: %.2f%% > thres: %.2f%%"), MaxLocalDiff * 100, MaximumLocalError * 100);
+		return EHLODRebuildPolicyDecision::ApproveRebuild;
+	}
+
+	OutReason = FString::Printf(TEXT("global diff: %.2f%% < thres: %.2f%%, local diff: %.2f%% < thres: %.2f%%"), MaxGlobalDiff * 100, MaximumGlobalError * 100, MaxLocalDiff * 100, MaximumLocalError * 100);
+	return EHLODRebuildPolicyDecision::RejectRebuild;
+}
+
+void UHLODRebuildPolicyImageData::PreSave(FObjectPreSaveContext InObjectSaveContext)
+{
+	Super::PreSave(InObjectSaveContext);
+
+	if (InObjectSaveContext.IsFromAutoSave())
+	{
+		return;
+	}
+
+	CompressData();
+}
+
+static const ANSICHAR	HLOD_IMAGE_COMPARISON_MAGIC_MARKER[] = "HLOD_IMAGE_COMPARISON_MAGIC_MARKER";
+static const uint32		HLOD_IMAGE_COMPARISON_DATA_VERSION = 0;
+
+void UHLODRebuildPolicyImageData::CompressData() const
+{
+	if (!bNeedsDataCompression)
+	{
+		return;
+	}
+
+	UHLODRebuildPolicyImageData& MutableThis = *const_cast<UHLODRebuildPolicyImageData*>(this);
+
+	// Compress all the images and store them
+	ParallelFor(CaptureFrames.Num(), [&](int32 AngleIndex)
+	{
+		FHLODCaptureFrame& CaptureFrame = MutableThis.CaptureFrames[AngleIndex];
+
+		for (uint32 ImageIndex = 0; ImageIndex < EHLODCaptureImageType::Num; ++ImageIndex)
+		{
+			FHLODCaptureImage& CaptureImage = CaptureFrame.Images[ImageIndex];
+
+			if (!CaptureImage.Pixels.IsEmpty())
+			{
+				uint8 DataVersion = HLOD_IMAGE_COMPARISON_DATA_VERSION;
+
+				FAnsiString ImageTypeString(StaticEnum<EHLODCaptureImageType>()->GetNameStringByValue(ImageIndex));
+				FAnsiString ImageName = FAnsiString::Printf("Angle%d_%s", AngleIndex, *ImageTypeString);
+				int32 ImageNameLen = ImageName.Len();
+				int32 UncompressedSize = CaptureImage.Pixels.Num() * CaptureImage.BytesPerPixel;
+
+				TArray<uint8> UncompressedBuffer;
+				UncompressedBuffer.SetNumUninitialized(UncompressedSize);
+
+				if (ImageIndex == Capture_AlphaMask)
+				{
+					for (int32 iPixel = 0; iPixel < CaptureImage.Pixels.Num(); ++iPixel)
+					{
+						UncompressedBuffer[iPixel] = CaptureImage.Pixels[iPixel].A;
+					}
+				}
+				else if (CaptureImage.BytesPerPixel == 1)
+				{
+					for (int32 iPixel = 0; iPixel < CaptureImage.Pixels.Num(); ++iPixel)
+					{
+						UncompressedBuffer[iPixel] = CaptureImage.Pixels[iPixel].R;
+					}
+				}
+				else if (CaptureImage.BytesPerPixel == 3)
+				{
+					for (int32 iPixel = 0; iPixel < CaptureImage.Pixels.Num(); ++iPixel)
+					{
+						UncompressedBuffer[iPixel * 3 + 0] = CaptureImage.Pixels[iPixel].R;
+						UncompressedBuffer[iPixel * 3 + 1] = CaptureImage.Pixels[iPixel].G;
+						UncompressedBuffer[iPixel * 3 + 2] = CaptureImage.Pixels[iPixel].B;
+					}
+				}
+				else
+				{
+					checkNoEntry();
+				}
+
+				TArray<uint8> CompressedBuffer;
+				CompressedBuffer.SetNumUninitialized(UncompressedSize);
+				int32 CompressedSize = UncompressedSize;
+				FCompression::CompressMemory(NAME_Zlib, CompressedBuffer.GetData(), CompressedSize, UncompressedBuffer.GetData(), UncompressedSize, COMPRESS_BiasSize);
+				CompressedBuffer.SetNum(CompressedSize);
+
+				// Write all the info to a buffer that can be read easily outside of UE
+				FMemoryWriter MemoryWriter(CaptureImage.PixelsCompressed);
+				MemoryWriter.Serialize((void*)HLOD_IMAGE_COMPARISON_MAGIC_MARKER, sizeof(HLOD_IMAGE_COMPARISON_MAGIC_MARKER) - 1);
+				MemoryWriter << DataVersion;
+				MemoryWriter << ImageNameLen;
+				MemoryWriter.Serialize((void*)*ImageName, ImageNameLen);
+				MemoryWriter << CaptureImage.Width;
+				MemoryWriter << CaptureImage.Height;
+				MemoryWriter << CaptureImage.BytesPerPixel;
+				MemoryWriter << UncompressedSize;
+				MemoryWriter << CompressedBuffer;
+			}
+			else
+			{
+				CaptureImage.PixelsCompressed.Empty();
+			}
+		}
+	});
+
+	MutableThis.bNeedsDataCompression = false;
+}
+
+bool UHLODRebuildPolicyImageData::DecompressData() const
+{
+	if (!bNeedsDataDecompression)
+	{
+		return true;
+	}
+
+	UHLODRebuildPolicyImageData& MutableThis = *const_cast<UHLODRebuildPolicyImageData*>(this);
+
+	TAtomic<uint32> ErrorCount(0);
+
+	ParallelFor(CaptureFrames.Num(), [&](int32 AngleIndex)
+	{
+		FHLODCaptureFrame& CaptureFrame = MutableThis.CaptureFrames[AngleIndex];
+
+		for (uint32 ImageIndex = 0; ImageIndex < EHLODCaptureImageType::Num; ++ImageIndex)
+		{
+			FHLODCaptureImage& CaptureImage = CaptureFrame.Images[ImageIndex];
+
+			// Clear any stale data
+			CaptureImage.Pixels.Reset();
+
+			// Nothing to do?
+			if (CaptureImage.PixelsCompressed.IsEmpty())
+			{
+				ErrorCount++;
+				continue;
+			}
+
+			FMemoryReader Reader(CaptureImage.PixelsCompressed, /*bIsPersistent*/ true);
+
+			// Validate header
+			if (Reader.TotalSize() < sizeof(HLOD_IMAGE_COMPARISON_MAGIC_MARKER) - 1)
+			{
+				UE_LOG(LogTemp, Warning, TEXT("HLOD image buffer too small (angle %d, image %d)."), AngleIndex, ImageIndex);
+				ErrorCount++;
+				continue;
+			}
+
+			ANSICHAR Header[sizeof(HLOD_IMAGE_COMPARISON_MAGIC_MARKER)] = { 0 };
+			Reader.Serialize(Header, sizeof(HLOD_IMAGE_COMPARISON_MAGIC_MARKER) - 1);
+
+			if (FCStringAnsi::Strncmp(Header, HLOD_IMAGE_COMPARISON_MAGIC_MARKER, sizeof(HLOD_IMAGE_COMPARISON_MAGIC_MARKER) - 1) != 0)
+			{
+				UE_LOG(LogTemp, Warning, TEXT("HLOD image header mismatch (angle %d, image %d)."), AngleIndex, ImageIndex);
+				ErrorCount++;
+				continue;
+			}
+
+			// Version - Currently unused but might be useful to extend/modify the stored data in the future
+			uint8 DataVersion = HLOD_IMAGE_COMPARISON_DATA_VERSION;
+			Reader << DataVersion;
+
+			// Image name (discarded after read; only used for external tooling)
+			int32 ImageNameLen = 0;
+			Reader << ImageNameLen;
+
+			if (ImageNameLen < 0 || (int64)Reader.Tell() + ImageNameLen >(int64)Reader.TotalSize())
+			{
+				UE_LOG(LogTemp, Warning, TEXT("Invalid image name length (angle %d, image %d)."), AngleIndex, ImageIndex);
+				ErrorCount++;
+				continue;
+			}
+
+			if (ImageNameLen > 0)
+			{
+				TArray<ANSICHAR> NameBytes;
+				NameBytes.SetNumUninitialized(ImageNameLen);
+				Reader.Serialize(NameBytes.GetData(), ImageNameLen);
+			}
+
+			// Dimensions
+			int32 Width;
+			int32 Height;
+			uint8 BytesPerPixel;
+			Reader << Width;
+			Reader << Height;
+			Reader << BytesPerPixel;
+
+			check(CaptureImage.Width == Width);
+			check(CaptureImage.Height == Height);
+			check(CaptureImage.BytesPerPixel == BytesPerPixel);
+
+			// Uncompressed size
+			int32 UncompressedSize = 0;
+			Reader << UncompressedSize;
+
+			// Compressed payload (serialized TArray<uint8>)
+			TArray<uint8> CompressedBuffer;
+			Reader << CompressedBuffer;
+
+			// Inflate
+			if (UncompressedSize <= 0 || CompressedBuffer.Num() <= 0)
+			{
+				UE_LOG(LogTemp, Warning, TEXT("Invalid sizes while decompressing HLOD image (angle %d, image %d)."), AngleIndex, ImageIndex);
+				ErrorCount++;
+				continue;
+			}
+
+			TArray<uint8> UncompressedBuffer;
+			UncompressedBuffer.SetNumUninitialized(UncompressedSize);
+			int32 OutSize = UncompressedSize;
+
+			const bool bOk = FCompression::UncompressMemory(NAME_Zlib, UncompressedBuffer.GetData(), OutSize, CompressedBuffer.GetData(),CompressedBuffer.Num());
+			if (!bOk || OutSize != UncompressedSize)
+			{
+				UE_LOG(LogTemp, Warning, TEXT("Failed to uncompress HLOD image (angle %d, image %d)."), AngleIndex, ImageIndex);
+				CaptureImage.Pixels.Empty();
+				ErrorCount++;
+				continue;
+			}
+
+			CaptureImage.Pixels.SetNumUninitialized(Width * Height);
+
+			if (ImageIndex == Capture_AlphaMask)
+			{
+				for (int32 iPixel = 0; iPixel < CaptureImage.Pixels.Num(); ++iPixel)
+				{
+					CaptureImage.Pixels[iPixel] = FColor(255, 255, 255, UncompressedBuffer[iPixel]);
+				}
+			}
+			else if (CaptureImage.BytesPerPixel == 1)
+			{
+				for (int32 iPixel = 0; iPixel < CaptureImage.Pixels.Num(); ++iPixel)
+				{
+					const uint8 R = UncompressedBuffer[iPixel];
+					const uint8 A = CaptureFrame.Images[Capture_AlphaMask].Pixels[iPixel].A; // Safe to dereference as it's the first image processed
+					CaptureImage.Pixels[iPixel] = FColor(R, R, R, A);
+				}
+			}
+			else if (CaptureImage.BytesPerPixel == 3)
+			{
+				for (int32 iPixel = 0; iPixel < CaptureImage.Pixels.Num(); ++iPixel)
+				{
+					const uint8 R = UncompressedBuffer[iPixel * CaptureImage.BytesPerPixel + 0];
+					const uint8 G = UncompressedBuffer[iPixel * CaptureImage.BytesPerPixel + 1];
+					const uint8 B = UncompressedBuffer[iPixel * CaptureImage.BytesPerPixel + 2];
+					const uint8 A = CaptureFrame.Images[Capture_AlphaMask].Pixels[iPixel].A; // Safe to dereference as it's the first image processed
+					CaptureImage.Pixels[iPixel] = FColor(R, G, B, A);
+				}
+			}
+			else
+			{
+				checkNoEntry();
+				ErrorCount++;
+			}
+		}
+	});
+
+	MutableThis.bNeedsDataDecompression = false;
+
+	return ErrorCount == 0;
+}
diff --git a/Engine/Plugins/Editor/WorldPartitionHLODUtilities/Source/Private/WorldPartition/HLOD/Utilities/WorldPartitionHLODUtilities.cpp b/Engine/Plugins/Editor/WorldPartitionHLODUtilities/Source/Private/WorldPartition/HLOD/Utilities/WorldPartitionHLODUtilities.cpp
--- a/Engine/Plugins/Editor/WorldPartitionHLODUtilities/Source/Private/WorldPartition/HLOD/Utilities/WorldPartitionHLODUtilities.cpp
+++ b/Engine/Plugins/Editor/WorldPartitionHLODUtilities/Source/Private/WorldPartition/HLOD/Utilities/WorldPartitionHLODUtilities.cpp
@@ -12,7 +12,6 @@
 #include "WorldPartition/HLOD/CustomHLODActor.h"
 #include "WorldPartition/HLOD/HLODActor.h"
 #include "WorldPartition/HLOD/HLODCreationFilter.h"
-#include "WorldPartition/HLOD/HLODHashBuilder.h"
 #include "WorldPartition/HLOD/HLODLayer.h"
 #include "WorldPartition/HLOD/HLODModifier.h"
 #include "WorldPartition/HLOD/HLODSourceActorsFromCell.h"
@@ -54,32 +53,6 @@ TAutoConsoleVariable<bool> CVarHLODUseLegacy32BitNameHash(
 	TEXT("Force the use of 32-bit hashing to use when generating HLOD actor names. Switching to 64-bit will result in new names/new actor files."),
 	ECVF_ReadOnly);
 
-static void ComputeHLODHash(FHLODHashBuilder& InHashBuilder, const AWorldPartitionHLOD* InHLODActor, const TArray<UActorComponent*>& InSourceComponents)
-{
-	// Base key, changing this will force a rebuild of all HLODs
-	FString HLODBaseKey = TEXT("0D33837AB1A04CC2AEEB04C5C217DE96");
-	InHashBuilder.HashField(HLODBaseKey, TEXT("HLODBaseKey"));
-
-	// HLOD Source Actors
-	if (ensure(InHLODActor->GetSourceActors()))
-	{
-		InHLODActor->GetSourceActors()->ComputeHLODHash(InHashBuilder);
-	}
-	
-	// Min Visible Distance
-	InHashBuilder.HashField(InHLODActor->GetMinVisibleDistance(), TEXT("MinVisibleDistance"));
-
-	// ISM Component Class
-	TSubclassOf<UHLODInstancedStaticMeshComponent> HLODISMComponentClass = UHLODBuilder::GetInstancedStaticMeshComponentClass();
-	if (HLODISMComponentClass != UHLODInstancedStaticMeshComponent::StaticClass())
-	{
-		InHashBuilder.HashField(HLODISMComponentClass->GetPathName(), TEXT("HLODISMComponentClass"));
-	}
-
-	// Append all components hashes
-	UHLODBuilder::ComputeHLODHash(InHashBuilder, InSourceComponents);
-}
-
 struct FPerHLODLayerKey
 {
 	UHLODLayer* Layer;
@@ -798,11 +771,14 @@ void UnloadSourceActors(ULevelStreaming* InLevelStreaming)
 		// Destroy the package world and remove it from root
 		UPackage* Package = Level->GetPackage();
 		UWorld* PackageWorld = UWorld::FindWorldInPackage(Package);
-		PackageWorld->DestroyWorld(false);
+		if (ensure(PackageWorld))
+		{
+			PackageWorld->DestroyWorld(false);
+		}
 	}
 }
 
-uint32 FWorldPartitionHLODUtilities::BuildHLOD(AWorldPartitionHLOD* InHLODActor)
+bool FWorldPartitionHLODUtilities::BuildHLOD(const FHLODBuildParams& InBuildParams)
 {
 	FAutoScopedDurationTimer TotalTimeScope;
 
@@ -810,6 +786,8 @@ uint32 FWorldPartitionHLODUtilities::BuildHLOD(AWorldPartitionHLOD* InHLODActor)
 	int64 LoadTimeMS = 0;
 	int64 BuildTimeMS = 0;
 	int64 TotalTimeMS = 0;
+
+	AWorldPartitionHLOD* HLODActor = InBuildParams.HLODActor;
 	
 	// Load actors relevant to HLODs
 	bool bIsDirty = false;
@@ -828,7 +806,7 @@ uint32 FWorldPartitionHLODUtilities::BuildHLOD(AWorldPartitionHLOD* InHLODActor)
 	{
 		FAutoScopedDurationTimer LoadTimeScope;
 
-		bool bLoadSucceeded = LoadSourceActors(InHLODActor, TempWorld, bIsDirty);
+		bool bLoadSucceeded = LoadSourceActors(HLODActor, TempWorld, bIsDirty);
 		if (!ensure(bLoadSucceeded))
 		{
 			UE_LOG(LogHLODBuilder, Error, TEXT("FWorldPartitionHLODUtilities::BuildHLOD() - Failed to load source actors"));
@@ -838,37 +816,45 @@ uint32 FWorldPartitionHLODUtilities::BuildHLOD(AWorldPartitionHLOD* InHLODActor)
 		LoadTimeMS = FMath::RoundToInt(LoadTimeScope.GetTime() * 1000);
 	}
 
+	const bool bForceBuild = bIsDirty || InBuildParams.bForceBuild;
+
 	const UWorldPartitionHLODSourceActors* SourceActors = InHLODActor->GetSourceActors();
 	check(SourceActors);
 	TArray<UActorComponent*> HLODRelevantComponents = GatherHLODRelevantComponents(TempWorld, *SourceActors);
-	FHLODHashBuilder HashBuilder;
-	::ComputeHLODHash(HashBuilder, InHLODActor, HLODRelevantComponents);
 
-	uint32 NewHLODHash = HashBuilder.GetCrc();
-	uint32 OldHLODHash = bIsDirty ? 0 : InHLODActor->GetHLODHash();
+	FHLODRebuildPolicyDataSet OldRebuildPolicyDataSet = bForceBuild ? FHLODRebuildPolicyDataSet() : HLODActor->GetHLODRebuildPolicyDataSet();
+	FHLODRebuildPolicyDataSet NewRebuildPolicyDataSet = UHLODRebuildPolicy::ComputeDataForRebuildPolicies(HLODActor, HLODRelevantComponents, /*bInForComparison=*/true);
+	EHLODRebuildPolicyDecision RebuildDecision;
 
 	if (HLODBuildEvaluatorDelegate.IsBound())
 	{
-		bool bMustPerformBuild = HLODBuildEvaluatorDelegate.Execute(InHLODActor, OldHLODHash, NewHLODHash);
-		if (!bMustPerformBuild)
-		{
-			return OldHLODHash;
-		}
+		RebuildDecision = HLODBuildEvaluatorDelegate.Execute(HLODActor, OldRebuildPolicyDataSet, NewRebuildPolicyDataSet);
+	}
+	else
+	{
+		RebuildDecision = UHLODRebuildPolicy::Evaluate(HLODActor, OldRebuildPolicyDataSet, NewRebuildPolicyDataSet);
+	}
+
+	if (RebuildDecision != EHLODRebuildPolicyDecision::ApproveRebuild)
+	{
+		UE_LOG(LogHLODBuilder, Verbose, TEXT("HLOD actor \"%s\" doesn't need to be rebuilt."), *HLODActor->GetActorLabel());
+		return false;
 	}
-	else if (OldHLODHash == NewHLODHash)
+
+	// At this point, we know the HLOD would need to be rebuilt
+	if (InBuildParams.bTestOnly)
 	{
-		UE_LOG(LogHLODBuilder, Verbose, TEXT("HLOD actor \"%s\" doesn't need to be rebuilt."), *InHLODActor->GetActorLabel());
-		return OldHLODHash;
+		return true;
 	}
 
 	// Clear stats as we're about to refresh them
-	InHLODActor->ResetStats();
+	HLODActor->ResetStats();
 
 	// Rename previous assets found in the HLOD actor package.
 	// Move the previous asset(s) to the transient package, to avoid any object reuse during the build
 	{
 	    TArray<UObject*> ObjectsToRename;
-	    ForEachObjectWithOuter(InHLODActor->GetPackage(), [&ObjectsToRename](UObject* Obj)
+	    ForEachObjectWithOuter(HLODActor->GetPackage(), [&ObjectsToRename](UObject* Obj)
 	    {
 		    if (!Obj->IsA<AActor>() && !Obj->IsA<UWorld>())
 		    {
@@ -887,9 +873,9 @@ uint32 FWorldPartitionHLODUtilities::BuildHLOD(AWorldPartitionHLOD* InHLODActor)
 	}
 
 	// Gather stats from the input to our HLOD build
-	GatherInputStats(InHLODActor, HLODRelevantComponents);
+	GatherInputStats(HLODActor, HLODRelevantComponents);
 	
-	const UHLODLayer* HLODLayer = InHLODActor->GetSourceActors()->GetHLODLayer();
+	const UHLODLayer* HLODLayer = HLODActor->GetSourceActors()->GetHLODLayer();
 	TSubclassOf<UHLODBuilder> HLODBuilderClass = GetHLODBuilderClass(HLODLayer);
 
 	if (HLODBuilderClass)
@@ -902,12 +888,12 @@ uint32 FWorldPartitionHLODUtilities::BuildHLOD(AWorldPartitionHLOD* InHLODActor)
 			HLODBuilder->SetHLODBuilderSettings(HLODLayer->GetHLODBuilderSettings());
 
 			FHLODBuildContext HLODBuildContext;
-			HLODBuildContext.World = InHLODActor->GetWorld();
+			HLODBuildContext.World = HLODActor->GetWorld();
 			HLODBuildContext.SourceComponents = HLODRelevantComponents;
-			HLODBuildContext.AssetsOuter = InHLODActor->GetPackage();
-			HLODBuildContext.AssetsBaseName = InHLODActor->GetActorLabel();
-			HLODBuildContext.MinVisibleDistance = InHLODActor->GetMinVisibleDistance();
-			HLODBuildContext.WorldPosition = InHLODActor->GetHLODBounds().GetCenter();
+			HLODBuildContext.AssetsOuter = HLODActor->GetPackage();
+			HLODBuildContext.AssetsBaseName = HLODActor->GetActorLabel();
+			HLODBuildContext.MinVisibleDistance = HLODActor->GetMinVisibleDistance();
+			HLODBuildContext.WorldPosition = HLODActor->GetHLODBounds().GetCenter();
 
 			const TSubclassOf<UWorldPartitionHLODModifier> HLODModifierClass = HLODLayer->GetHLODModifierClass();
 			UWorldPartitionHLODModifier* HLODModifier = HLODModifierClass.Get() ? NewObject<UWorldPartitionHLODModifier>(GetTransientPackage(), HLODModifierClass) : nullptr;
@@ -933,7 +919,7 @@ uint32 FWorldPartitionHLODUtilities::BuildHLOD(AWorldPartitionHLOD* InHLODActor)
 
 			if (BuildResult.HLODComponents.IsEmpty())
 			{
-				UE_LOG(LogHLODBuilder, Log, TEXT("HLOD generation created no component for HLOD actor %s. Listing source components:"), *InHLODActor->GetActorLabel());
+				UE_LOG(LogHLODBuilder, Log, TEXT("HLOD generation created no component for HLOD actor %s. Listing source components:"), *HLODActor->GetActorLabel());
 				for (UActorComponent* SourceComponent : HLODBuildContext.SourceComponents)
 				{
 					UE_LOG(LogHLODBuilder, Log, TEXT("\t* %s"), *SourceComponent->GetPathName());
@@ -975,7 +961,7 @@ uint32 FWorldPartitionHLODUtilities::BuildHLOD(AWorldPartitionHLOD* InHLODActor)
 						if (StaticMesh->GetPackage() == HLODBuildContext.AssetsOuter)
 						{
 							// Set up ray tracing far fields for always loaded HLODs
-							if (!InHLODActor->GetIsSpatiallyLoaded() && StaticMesh->bSupportRayTracing)
+							if (!HLODActor->GetIsSpatiallyLoaded() && StaticMesh->bSupportRayTracing)
 							{
 								StaticMeshComponent->bRayTracingFarField = true;
 							}
@@ -1043,56 +1029,25 @@ uint32 FWorldPartitionHLODUtilities::BuildHLOD(AWorldPartitionHLOD* InHLODActor)
 				}
 			}
 
-			InHLODActor->SetInputStats(BuildResult.InputStats);
-			InHLODActor->SetHLODComponents(BuildResult.HLODComponents);
+			HLODActor->SetInputStats(BuildResult.InputStats);
+			HLODActor->SetHLODComponents(BuildResult.HLODComponents);
 		}
 	}
 
 	// Gather stats pertaining to the assets generated during this build
-	GatherOutputStats(InHLODActor);
+	GatherOutputStats(HLODActor);
 
 	TotalTimeMS = FMath::RoundToInt(TotalTimeScope.GetTime() * 1000);
 
 	// Build timings stats
-	InHLODActor->SetStat(FWorldPartitionHLODStats::BuildTimeLoadMilliseconds, LoadTimeMS);
-	InHLODActor->SetStat(FWorldPartitionHLODStats::BuildTimeBuildMilliseconds, BuildTimeMS);
-	InHLODActor->SetStat(FWorldPartitionHLODStats::BuildTimeTotalMilliseconds, TotalTimeMS);
-
-	// Update HLOD hash and provide the detailed hash report to help future investigation into what caused HLODs to be rebuilt.
-	InHLODActor->SetHLODHash(NewHLODHash, HashBuilder.BuildHashReport());
-
-	return NewHLODHash;
-}
-
-uint32 FWorldPartitionHLODUtilities::ComputeHLODHash(const AWorldPartitionHLOD* InHLODActor)
-{
-	const bool bInformEngineOfWorld = false;
-	const bool bAddToRoot = true;
-	UPackage* TempPackage = NewObject<UPackage>(nullptr, *MakeUniqueObjectName(nullptr, UPackage::StaticClass(), TEXT("/Temp/BuildHLODPackage")).ToString(), RF_Transient);
-	UWorld* TempWorld = UWorld::CreateWorld(EWorldType::Editor, bInformEngineOfWorld, TEXT("BuildHLODWorld"), TempPackage, bAddToRoot);
-	check(TempWorld);
-
-	ON_SCOPE_EXIT
-	{
-		TempWorld->DestroyWorld(bInformEngineOfWorld);
-	};
-
-	// Load actors relevant to HLODs
-	bool bIsDirty = false;
-	bool bLoadSucceeded = LoadSourceActors(InHLODActor, TempWorld, bIsDirty);
-	if (!ensure(bLoadSucceeded))
-	{
-		UE_LOG(LogHLODBuilder, Error, TEXT("FWorldPartitionHLODUtilities::ComputeHLODHash() - Failed to load source actors"));
-		return 0;
-	}
+	HLODActor->SetStat(FWorldPartitionHLODStats::BuildTimeLoadMilliseconds, LoadTimeMS);
+	HLODActor->SetStat(FWorldPartitionHLODStats::BuildTimeBuildMilliseconds, BuildTimeMS);
+	HLODActor->SetStat(FWorldPartitionHLODStats::BuildTimeTotalMilliseconds, TotalTimeMS);
 
-	const UWorldPartitionHLODSourceActors* SourceActors = InHLODActor->GetSourceActors();
-	check(SourceActors);
-	TArray<UActorComponent*> HLODRelevantComponents = GatherHLODRelevantComponents(TempWorld, *SourceActors);
+	FHLODRebuildPolicyDataSet UpdatedRebuildPolicyDataSet = UHLODRebuildPolicy::ComputeDataForRebuildPolicies(HLODActor, HLODRelevantComponents, /*bInForComparison=*/false);
+	HLODActor->SetHLODRebuildPolicyDataSet(UpdatedRebuildPolicyDataSet);
 
-	FHLODHashBuilder HashBuilder;
-	::ComputeHLODHash(HashBuilder, InHLODActor, HLODRelevantComponents);
-	return HashBuilder.GetCrc();
+	return true;
 }
 
 void FWorldPartitionHLODUtilities::SetHLODBuildEvaluator(IWorldPartitionHLODUtilities::FHLODBuildEvaluator BuildEvaluatorDelegate)
@@ -1100,4 +1055,4 @@ void FWorldPartitionHLODUtilities::SetHLODBuildEvaluator(IWorldPartitionHLODUtil
 	HLODBuildEvaluatorDelegate = BuildEvaluatorDelegate;
 }
 
-#endif
+#endif
\ No newline at end of file
diff --git a/Engine/Plugins/Editor/WorldPartitionHLODUtilities/Source/Public/WorldPartition/HLOD/RebuildPolicies/HLODRebuildPolicyHashCompare.h b/Engine/Plugins/Editor/WorldPartitionHLODUtilities/Source/Public/WorldPartition/HLOD/RebuildPolicies/HLODRebuildPolicyHashCompare.h
new file mode 100644
--- /dev/null
+++ b/Engine/Plugins/Editor/WorldPartitionHLODUtilities/Source/Public/WorldPartition/HLOD/RebuildPolicies/HLODRebuildPolicyHashCompare.h
@@ -0,0 +1,45 @@
+// Copyright Epic Games, Inc. All Rights Reserved.
+
+#pragma once
+
+#include "WorldPartition/HLOD/HLODHashBuilder.h"
+#include "WorldPartition/HLOD/HLODRebuildPolicy.h"
+#include "HLODRebuildPolicyHashCompare.generated.h"
+
+UCLASS()
+class UHLODRebuildPolicyHashData : public UHLODRebuildPolicyData
+{
+	GENERATED_BODY()
+
+public:
+	virtual FString ToString() const override { return HLODHashReport; };
+
+public:
+	UPROPERTY()
+	uint32 HLODHash;
+
+	// Transient, will be aggregated into the HLOD actor build report
+	UPROPERTY(Transient)
+	FString HLODHashReport;
+};
+
+/**
+ * This HLOD rebuild policy computes a hash of various elements in order to detect any change that might
+ * impact the resulting HLODs. It is, in effect, extremely sensitive.
+ * It will:
+ *		APPROVE a rebuild - if the old hash differs from the new hash
+ *		REJECT  a rebuild - if the old hash equals the new hash
+ */
+UCLASS()
+class UHLODRebuildPolicyHashCompare : public UHLODRebuildPolicy
+{
+	GENERATED_BODY()
+
+protected:
+	virtual TSubclassOf<UHLODRebuildPolicyData> GetRebuildPolicyDataType() const override { return UHLODRebuildPolicyHashData::StaticClass(); }
+	virtual EHLODRebuildPolicyDecision Evaluate(const AWorldPartitionHLOD* InHLODActor, const UHLODRebuildPolicyData* InOldData, const UHLODRebuildPolicyData* InNewData, FString& OutReason) const override;
+	virtual UHLODRebuildPolicyData* ComputeDataForRebuildPolicy(const AWorldPartitionHLOD* InHLODActor, const TArray<UActorComponent*>& InSourceComponents, const UHLODRebuildPolicyData* InExistingPolicyData) const override;
+
+private:
+	void ComputeHLODHash(FHLODHashBuilder& InHashBuilder, const AWorldPartitionHLOD* InHLODActor, const TArray<UActorComponent*>& InSourceComponents) const;
+};
diff --git a/Engine/Plugins/Editor/WorldPartitionHLODUtilities/Source/Public/WorldPartition/HLOD/RebuildPolicies/HLODRebuildPolicyImageCompare.h b/Engine/Plugins/Editor/WorldPartitionHLODUtilities/Source/Public/WorldPartition/HLOD/RebuildPolicies/HLODRebuildPolicyImageCompare.h
new file mode 100644
--- /dev/null
+++ b/Engine/Plugins/Editor/WorldPartitionHLODUtilities/Source/Public/WorldPartition/HLOD/RebuildPolicies/HLODRebuildPolicyImageCompare.h
@@ -0,0 +1,127 @@
+// Copyright Epic Games, Inc. All Rights Reserved.
+
+#pragma once
+
+#include "WorldPartition/HLOD/HLODRebuildPolicy.h"
+#include "HLODRebuildPolicyImageCompare.generated.h"
+
+
+UENUM()
+enum EHLODCaptureImageType : uint8
+{
+	Capture_AlphaMask,	// 8-bits
+	Capture_BaseColor,	// 24-bits
+	Capture_Normal,		// 24-bits
+	Capture_Emissive,	// 24-bits
+	Capture_Metallic,	// 8-bits
+	Capture_Roughness,	// 8-bits
+	Capture_Specular,	// 8-bits
+	Num
+};
+
+template<EHLODCaptureImageType T> struct THLODCaptureImageTraits { static constexpr int32 NumChannels = 1; };
+template<> struct THLODCaptureImageTraits<Capture_BaseColor>	 { static constexpr int32 NumChannels = 3; };
+template<> struct THLODCaptureImageTraits<Capture_Normal>		 { static constexpr int32 NumChannels = 3; };
+template<> struct THLODCaptureImageTraits<Capture_Emissive>		 { static constexpr int32 NumChannels = 3; };
+
+USTRUCT()
+struct FHLODCaptureImage
+{
+	GENERATED_BODY()
+
+	UPROPERTY()
+	int32 Width;
+
+	UPROPERTY()
+	int32 Height;
+
+	UPROPERTY()
+	uint8 BytesPerPixel;
+
+	UPROPERTY()
+	TArray<uint8> PixelsCompressed;	// Compressed, minimal number of channels (BytesPerPixel), saved to disk
+
+	UPROPERTY(Transient)
+	TArray<FColor> Pixels;			// Uncompressed, RGBA, not saved to disk
+};
+
+USTRUCT()
+struct FHLODCaptureFrame
+{
+	GENERATED_BODY()
+
+	UPROPERTY()
+	FVector CameraLocation;
+
+	UPROPERTY()
+	FRotator CameraRotation;
+
+	UPROPERTY()
+	float OrthoWidth;
+
+	UPROPERTY()
+	FHLODCaptureImage Images[EHLODCaptureImageType::Num];
+};
+
+UCLASS()
+class UHLODRebuildPolicyImageData : public UHLODRebuildPolicyData
+{
+	GENERATED_BODY()
+
+public:
+	UPROPERTY()
+	FBox CaptureBounds;
+
+	UPROPERTY()
+	TArray<FHLODCaptureFrame> CaptureFrames;
+
+	UPROPERTY()
+	uint32 HLODLayerHash;
+
+	UPROPERTY(Transient)
+	bool bNeedsDataCompression = false;
+
+	UPROPERTY(Transient)
+	bool bNeedsDataDecompression = true;
+
+	void CompressData() const;
+	bool DecompressData() const;
+
+protected:
+	//~ Begin UObject Interface.
+	virtual void PreSave(FObjectPreSaveContext InObjectSaveContext) override;
+	//~ End UObject Interface.
+};
+
+/**
+ * This HLOD rebuild policy takes screen captures of the SOURCE actors over multiple angles (and for different GBuffer properties)
+ * It will then perform an SSIM evaluation in order to assess if there is a significant visual change between the old and new data set.
+ * It will:
+ *		APPROVE a rebuild - if the difference between old vs new images is greater than the configured threshold
+ *		REJECT  a rebuild - if the difference between old vs new images is less than the configured threshold
+ */
+UCLASS(MinimalAPI, Config = Editor)
+class UHLODRebuildPolicyImageCompare : public UHLODRebuildPolicy
+{
+	GENERATED_UCLASS_BODY()
+
+protected:
+	virtual TSubclassOf<UHLODRebuildPolicyData> GetRebuildPolicyDataType() const override { return UHLODRebuildPolicyImageData::StaticClass(); }
+	virtual EHLODRebuildPolicyDecision Evaluate(const AWorldPartitionHLOD* InHLODActor, const UHLODRebuildPolicyData* InOldData, const UHLODRebuildPolicyData* InNewData, FString& OutReason) const override;
+	virtual UHLODRebuildPolicyData* ComputeDataForRebuildPolicy(const AWorldPartitionHLOD* InHLODActor, const TArray<UActorComponent*>& InSourceComponents, const UHLODRebuildPolicyData* InExistingPolicyData) const override;
+
+	UPROPERTY(Config)
+	uint32 MaxImageSize = 512;
+
+	UPROPERTY(Config)
+	uint32 WarmupFrameCount = 5;
+
+	UPROPERTY(Config)
+	float CaptureImageMargin = 0.10f;
+
+	UPROPERTY(Config)
+	float MaximumLocalError = 1.00f;
+
+	UPROPERTY(Config)
+	float MaximumGlobalError = 0.05f;
+};
diff --git a/Engine/Plugins/Editor/WorldPartitionHLODUtilities/Source/Public/WorldPartition/HLOD/Utilities/WorldPartitionHLODUtilities.h b/Engine/Plugins/Editor/WorldPartitionHLODUtilities/Source/Public/WorldPartition/HLOD/Utilities/WorldPartitionHLODUtilities.h
--- a/Engine/Plugins/Editor/WorldPartitionHLODUtilities/Source/Public/WorldPartition/HLOD/Utilities/WorldPartitionHLODUtilities.h
+++ b/Engine/Plugins/Editor/WorldPartitionHLODUtilities/Source/Public/WorldPartition/HLOD/Utilities/WorldPartitionHLODUtilities.h
@@ -11,12 +11,11 @@ class FWorldPartitionHLODUtilities : public IWorldPartitionHLODUtilities
 {
 public:
 	virtual TArray<AWorldPartitionHLOD*> CreateHLODActors(FHLODCreationContext& InCreationContext, const FHLODCreationParams& InCreationParams, const TArray<IStreamingGenerationContext::FActorInstance>& InActors) override;
-	virtual uint32 BuildHLOD(AWorldPartitionHLOD* InHLODActor) override;
-	virtual uint32 ComputeHLODHash(const AWorldPartitionHLOD* InHLODActor) override;
+	virtual bool BuildHLOD(const FHLODBuildParams& InBuildParams) override;
 	virtual TSubclassOf<UHLODBuilder> GetHLODBuilderClass(const UHLODLayer* InHLODLayer) override;
 	virtual UHLODBuilderSettings* CreateHLODBuilderSettings(UHLODLayer* InHLODLayer) override;
 	virtual void SetHLODBuildEvaluator(IWorldPartitionHLODUtilities::FHLODBuildEvaluator BuildEvaluatorDelegate) override;
 
 protected:
-	IWorldPartitionHLODUtilities::FHLODBuildEvaluator HLODBuildEvaluatorDelegate;	
+	IWorldPartitionHLODUtilities::FHLODBuildEvaluator HLODBuildEvaluatorDelegate;
 };
diff --git a/Engine/Plugins/Editor/WorldPartitionHLODUtilities/Source/WorldPartitionHLODUtilities.Build.cs b/Engine/Plugins/Editor/WorldPartitionHLODUtilities/Source/WorldPartitionHLODUtilities.Build.cs
--- a/Engine/Plugins/Editor/WorldPartitionHLODUtilities/Source/WorldPartitionHLODUtilities.Build.cs
+++ b/Engine/Plugins/Editor/WorldPartitionHLODUtilities/Source/WorldPartitionHLODUtilities.Build.cs
@@ -25,10 +25,20 @@ public class WorldPartitionHLODUtilities : ModuleRules
             new string[]
 			{
 				"Engine",
+				"ImageWrapper",
+				"ImageWriteQueue",
+				"Json",
+				"JsonUtilities",
 				"MaterialBaking",
 				"MaterialUtilities",
 				"MeshDescription",
 				"MeshMergeUtilities",
+				"RenderCore",
+				"RHI",
+				"ScreenShotComparison",
+				"ScreenShotComparisonTools",
+				"Slate", 
+				"SlateCore",
 				"StaticMeshDescription",
 				"WorldPartitionEditor"
 			}
diff --git a/Engine/Source/Editor/UnrealEd/Private/WorldPartition/WorldPartitionHLODsBuilder.cpp b/Engine/Source/Editor/UnrealEd/Private/WorldPartition/WorldPartitionHLODsBuilder.cpp
--- a/Engine/Source/Editor/UnrealEd/Private/WorldPartition/WorldPartitionHLODsBuilder.cpp
+++ b/Engine/Source/Editor/UnrealEd/Private/WorldPartition/WorldPartitionHLODsBuilder.cpp
@@ -1772,24 +1772,31 @@ struct FScopedLogSuppression
 private:
 	FString LogCommand;
 };
-bool UWorldPartitionHLODsBuilder::EvaluateHLODBuildConditions(AWorldPartitionHLOD* HLODActor, uint32 OldHash, uint32 NewHash)
+
+EHLODRebuildPolicyDecision UWorldPartitionHLODsBuilder::EvaluateHLODBuildConditions(AWorldPartitionHLOD* HLODActor, const FHLODRebuildPolicyDataSet& InOldData, const FHLODRebuildPolicyDataSet& InNewData)
 {
 	TRACE_CPUPROFILER_EVENT_SCOPE(UWorldPartitionHLODsBuilder::EvaluateHLODBuildConditions);
 
-	BuildEvaluationStats::FStats& EvaluationStats = BuildEvaluationStats::StatsByHLODLayer.FindOrAdd(HLODActor->GetSourceActors()->GetHLODLayer()->GetFName());
+	FName HLODLayerStatsName = NAME_None;
+	if (ensure(HLODActor->GetSourceActors() && HLODActor->GetSourceActors()->GetHLODLayer()))
+	{
+		HLODLayerStatsName = HLODActor->GetSourceActors()->GetHLODLayer()->GetFName();
+	}
+
+	BuildEvaluationStats::FStats& EvaluationStats = BuildEvaluationStats::StatsByHLODLayer.FindOrAdd(HLODLayerStatsName);
 
 	FScopedDurationTimer TimerEvaluate(EvaluationStats.TimeEvaluate);
 
 	EvaluationStats.Evaluated++;
 
-	bool bShouldPerformBuild = OldHash != NewHash;
-	EvaluationStats.Built += bShouldPerformBuild ? 1 : 0;
+	EHLODRebuildPolicyDecision RebuildDecision = UHLODRebuildPolicy::Evaluate(HLODActor, InOldData, InNewData);
+	EvaluationStats.Built += RebuildDecision == EHLODRebuildPolicyDecision::ApproveRebuild ? 1 : 0;
 
-	// Parent branch hash comparison can't be performed without SCC
+	// Parent branch comparison can't be performed without SCC
 	if (!ISourceControlModule::Get().IsEnabled())
 	{
-		UE_LOG(LogWorldPartitionHLODsBuilder, Verbose, TEXT("Skipping hash comparison with parent branch: Source control is disabled"));
-		return bShouldPerformBuild;
+		UE_LOG(LogWorldPartitionHLODsBuilder, Verbose, TEXT("Skipping comparison with parent branch: Source control is disabled"));
+		return RebuildDecision;
 	}
 
 	// Current Branch
@@ -1797,8 +1804,8 @@ bool UWorldPartitionHLODsBuilder::EvaluateHLODBuildConditions(AWorldPartitionHLO
 	int32 CurrentBranchIndex = ISourceControlModule::Get().GetProvider().GetStateBranchIndex(CurrentBranch);
 	if (CurrentBranchIndex == INDEX_NONE)
 	{
-		UE_LOG(LogWorldPartitionHLODsBuilder, Verbose, TEXT("Skipping hash comparison with parent branch: Couldn't retrieve current branch index"));
-		return bShouldPerformBuild;
+		UE_LOG(LogWorldPartitionHLODsBuilder, Verbose, TEXT("Skipping comparison with parent branch: Couldn't retrieve current branch index"));
+		return RebuildDecision;
 	}
 	UE_LOG(LogWorldPartitionHLODsBuilder, VeryVerbose, TEXT("Current Branch: %s"), *CurrentBranch);
 
@@ -1807,8 +1814,8 @@ bool UWorldPartitionHLODsBuilder::EvaluateHLODBuildConditions(AWorldPartitionHLO
 	FString ParentBranch;
 	if (!ISourceControlModule::Get().GetProvider().GetStateBranchAtIndex(ParentBranchIndex, ParentBranch))
 	{
-		UE_LOG(LogWorldPartitionHLODsBuilder, Verbose, TEXT("Skipping hash comparison with parent branch: Couldn't retrieve parent branch"));
-		return bShouldPerformBuild;
+		UE_LOG(LogWorldPartitionHLODsBuilder, Verbose, TEXT("Skipping comparison with parent branch: Couldn't retrieve parent branch"));
+		return RebuildDecision;
 	}
 	UE_LOG(LogWorldPartitionHLODsBuilder, VeryVerbose, TEXT("Parent Branch: %s"), *ParentBranch);
 
@@ -1821,8 +1828,8 @@ bool UWorldPartitionHLODsBuilder::EvaluateHLODBuildConditions(AWorldPartitionHLO
 	ECommandResult::Type WhereResult = ISourceControlModule::Get().GetProvider().Execute(WhereOperation, LocalFilePathCurrent, EConcurrency::Synchronous);
 	if (WhereResult != ECommandResult::Succeeded || WhereOperation->GetFiles().IsEmpty())
 	{
-		UE_LOG(LogWorldPartitionHLODsBuilder, Verbose, TEXT("Skipping hash comparison with parent branch: Couldn't retrieve depot file path"));
-		return bShouldPerformBuild;
+		UE_LOG(LogWorldPartitionHLODsBuilder, Verbose, TEXT("Skipping comparison with parent branch: Couldn't retrieve depot file path"));
+		return RebuildDecision;
 	}
 	const FString DepotFilePathCurrent = WhereOperation->GetFiles()[0].RemotePath;
 	UE_LOG(LogWorldPartitionHLODsBuilder, VeryVerbose, TEXT("Depot File Path (Current): %s"), *DepotFilePathCurrent);
@@ -1831,8 +1838,8 @@ bool UWorldPartitionHLODsBuilder::EvaluateHLODBuildConditions(AWorldPartitionHLO
 	FString DepotFilePathParent = DepotFilePathCurrent;
 	if (!DepotFilePathParent.RemoveFromStart(CurrentBranch))
 	{
-		UE_LOG(LogWorldPartitionHLODsBuilder, Verbose, TEXT("Skipping hash comparison with parent branch: Couldn't map current depot file path to the parent branch file path"));
-		return bShouldPerformBuild;
+		UE_LOG(LogWorldPartitionHLODsBuilder, Verbose, TEXT("Skipping comparison with parent branch: Couldn't map current depot file path to the parent branch file path"));
+		return RebuildDecision;
 	}
 	DepotFilePathParent = ParentBranch + DepotFilePathParent;
 	UE_LOG(LogWorldPartitionHLODsBuilder, VeryVerbose, TEXT("Depot File Path (Parent): %s"), *DepotFilePathParent);
@@ -1850,15 +1857,15 @@ bool UWorldPartitionHLODsBuilder::EvaluateHLODBuildConditions(AWorldPartitionHLO
 		DownloadCommand->SetEnableInfoLogging(false);
 		if (ISourceControlModule::Get().GetProvider().Execute(DownloadCommand, DepotFilePathParent, EConcurrency::Synchronous) != ECommandResult::Succeeded)
 		{
-			UE_LOG(LogWorldPartitionHLODsBuilder, Verbose, TEXT("Skipping hash comparison with parent branch: Couldn't download file from parent branch"));
+			UE_LOG(LogWorldPartitionHLODsBuilder, Verbose, TEXT("Skipping comparison with parent branch: Couldn't download file from parent branch"));
 			EvaluationStats.NotFoundInParentBranch++;
-			return bShouldPerformBuild;
+			return RebuildDecision;
 		}
 		EvaluationStats.DataSyncedBytes += IFileManager::Get().FileSize(*LocalFilePathParent);
 	}
 
-	// Retrieve the parent hash
-	TOptional<uint32> ParentHash;
+	// Compare against the parent branch
+	TOptional<EHLODRebuildPolicyDecision> ParentRebuildDecision;
 
 	{
 		// Since we are copying the package from another (parent) branch, there's always the possibility that the package 
@@ -1891,7 +1898,7 @@ bool UWorldPartitionHLODsBuilder::EvaluateHLODBuildConditions(AWorldPartitionHLO
 				const FHLODActorDesc& ParentHLODActorDesc = *(static_cast<FHLODActorDesc*>(ParentActorDesc.Get()));
 				if (CurrentHLODActorDesc.EqualsUnbuilt(ParentHLODActorDesc))
 				{
-					ParentHash = FoundActor->GetHLODHash();
+					ParentRebuildDecision = UHLODRebuildPolicy::Evaluate(FoundActor, FoundActor->GetHLODRebuildPolicyDataSet(), InNewData);
 				}
 			}
 
@@ -1899,17 +1906,17 @@ bool UWorldPartitionHLODsBuilder::EvaluateHLODBuildConditions(AWorldPartitionHLO
 		}
 	}
 
-	const bool bOldHashMatch = OldHash == NewHash;
-	const bool bParentHashMatch = ParentHash.IsSet() && ParentHash.GetValue() == NewHash;
-	EvaluationStats.ReusedFromParentBranch += bParentHashMatch ? 1 : 0;
+	const bool bOldHLODMatch = RebuildDecision == EHLODRebuildPolicyDecision::RejectRebuild;
+	const bool bParentHLODMatch = ParentRebuildDecision.IsSet() && ParentRebuildDecision.GetValue() == EHLODRebuildPolicyDecision::RejectRebuild;
+	EvaluationStats.ReusedFromParentBranch += bParentHLODMatch ? 1 : 0;
 
-	// Abort build if old vs new hash match, or if the parent branch hlod has the same hash
-	const bool bAbortBuild = bOldHashMatch || bParentHashMatch;
+	// Abort build if old vs new HLOD match, or if the parent vs new match
+	const bool bAbortBuild = bOldHLODMatch || bParentHLODMatch;
 
 	// Copy parent if:
-	// -> Hash changed but matches parent (safe to reuse)
-	// -> Hash matches old and parent, but file content differs (check MD5)... we want identical data as much as possible to reduce patch sizes between versions
-	const bool bCopyParent = bParentHashMatch && (!bOldHashMatch || FMD5Hash::HashFile(*LocalFilePathCurrent) != FMD5Hash::HashFile(*LocalFilePathParent));
+	// -> Actor changed, but matches parent (safe to reuse)
+	// -> Content match in old vs parent, but file content differs (check MD5)... we want identical data as much as possible to reduce patch sizes between versions
+	const bool bCopyParent = bParentHLODMatch && (!bOldHLODMatch || FMD5Hash::HashFile(*LocalFilePathCurrent) != FMD5Hash::HashFile(*LocalFilePathParent));
 	if (bCopyParent)
 	{
 		ParentBranchHLODFileToCopy = LocalFilePathParent;
@@ -1921,13 +1928,13 @@ bool UWorldPartitionHLODsBuilder::EvaluateHLODBuildConditions(AWorldPartitionHLO
 	}
 
 		
-	EvaluationStats.Built -= bShouldPerformBuild && bAbortBuild ? 1 : 0;
+	EvaluationStats.Built -= RebuildDecision == EHLODRebuildPolicyDecision::ApproveRebuild && bAbortBuild ? 1 : 0;
 
 	UE_CLOG(bCopyParent, LogWorldPartitionHLODsBuilder, Display, TEXT("Skipping HLOD build - Copying HLOD actor file from parent branch: \"%s\""), *DepotFilePathParent);
-	UE_CLOG(!bCopyParent && bParentHashMatch, LogWorldPartitionHLODsBuilder, Display, TEXT("Skipping HLOD build - Identical to actor file in parent branch: \"%s\""), *DepotFilePathParent);
-	UE_CLOG(!bParentHashMatch && bOldHashMatch, LogWorldPartitionHLODsBuilder, Display, TEXT("Skipping HLOD build - HLOD actor doesn't need to be rebuilt: \"%s\""), *HLODActor->GetActorLabel());
+	UE_CLOG(!bCopyParent && bParentHLODMatch, LogWorldPartitionHLODsBuilder, Display, TEXT("Skipping HLOD build - Identical to actor file in parent branch: \"%s\""), *DepotFilePathParent);
+	UE_CLOG(!bParentHLODMatch && bOldHLODMatch, LogWorldPartitionHLODsBuilder, Display, TEXT("Skipping HLOD build - HLOD actor doesn't need to be rebuilt: \"%s\""), *HLODActor->GetActorLabel());
 	
-	return !bAbortBuild;
+	return bAbortBuild ? EHLODRebuildPolicyDecision::RejectRebuild : EHLODRebuildPolicyDecision::ApproveRebuild;
 }
 
 void UWorldPartitionHLODsBuilder::CopyParentBranchHLODFile(UPackage* HLODActorPackage)
diff --git a/Engine/Source/Editor/UnrealEd/Public/WorldPartition/WorldPartitionHLODsBuilder.h b/Engine/Source/Editor/UnrealEd/Public/WorldPartition/WorldPartitionHLODsBuilder.h
--- a/Engine/Source/Editor/UnrealEd/Public/WorldPartition/WorldPartitionHLODsBuilder.h
+++ b/Engine/Source/Editor/UnrealEd/Public/WorldPartition/WorldPartitionHLODsBuilder.h
@@ -3,6 +3,7 @@
 #pragma once
 
 #include "Misc/EnumClassFlags.h"
+#include "WorldPartition/HLOD/HLODRebuildPolicy.h"
 #include "WorldPartition/WorldPartitionBuilder.h"
 #include "WorldPartition/WorldPartitionBuilderHelpers.h"
 #include "WorldPartitionHLODsBuilder.generated.h"
@@ -77,7 +78,7 @@ protected:
 
 	void AllowExternalDataLayerInjection(const UWorld* InWorld, const UExternalDataLayerAsset* InExternalDataLayerAsset, bool& bOutAllowInjection);
 
-	bool EvaluateHLODBuildConditions(class AWorldPartitionHLOD* HLODActor, uint32 OldHash, uint32 NewHash);
+	EHLODRebuildPolicyDecision EvaluateHLODBuildConditions(class AWorldPartitionHLOD* HLODActor, const FHLODRebuildPolicyDataSet& InOldData, const FHLODRebuildPolicyDataSet& InNewData);
 	void CopyParentBranchHLODFile(UPackage* HLODActorPackage);
 
 private:
diff --git a/Engine/Source/Editor/WorldPartitionEditor/Private/WorldPartition/Customizations/WorldPartitionHLODDetailsCustomization.cpp b/Engine/Source/Editor/WorldPartitionEditor/Private/WorldPartition/Customizations/WorldPartitionHLODDetailsCustomization.cpp
--- a/Engine/Source/Editor/WorldPartitionEditor/Private/WorldPartition/Customizations/WorldPartitionHLODDetailsCustomization.cpp
+++ b/Engine/Source/Editor/WorldPartitionEditor/Private/WorldPartition/Customizations/WorldPartitionHLODDetailsCustomization.cpp
@@ -15,6 +15,7 @@
 #include "Dialogs/DlgPickPath.h"
 #include "Editor.h"
 #include "Engine/Selection.h"
+#include "Framework/Notifications/NotificationManager.h"
 #include "IContentBrowserDataModule.h"
 #include "IContentBrowserSingleton.h"
 #include "Interfaces/IMainFrameModule.h"
@@ -26,6 +27,7 @@
 #include "Widgets/DeclarativeSyntaxSupport.h"
 #include "Widgets/Input/SButton.h"
 #include "Widgets/Layout/SUniformGridPanel.h"
+#include "Widgets/Notifications/SNotificationList.h"
 #include "WorldPartition/WorldPartition.h"
 #include "WorldPartition/HLOD/HLODActor.h"
 
@@ -46,12 +48,25 @@ void FWorldPartitionHLODDetailsCustomization::CustomizeDetails(IDetailLayoutBuil
 		SNew(SVerticalBox)
 		+SVerticalBox::Slot()
 		[
-			SNew(SButton)
-			.HAlign(EHorizontalAlignment::HAlign_Center)
-			.VAlign(EVerticalAlignment::VAlign_Center)
-			.Text(LOCTEXT("BuildHLODButtonText", "Build HLOD"))
-			.OnClicked(this, &FWorldPartitionHLODDetailsCustomization::OnBuildHLOD)
-			.IsEnabled(this, &FWorldPartitionHLODDetailsCustomization::CanBuildHLOD)
+			SNew(SHorizontalBox)
+			+SHorizontalBox::Slot()
+			[
+				SNew(SButton)
+				.HAlign(EHorizontalAlignment::HAlign_Center)
+				.VAlign(EVerticalAlignment::VAlign_Center)
+				.Text(LOCTEXT("BuildHLODButtonText", "Build HLOD"))
+				.OnClicked(this, &FWorldPartitionHLODDetailsCustomization::OnBuildHLOD)
+				.IsEnabled(this, &FWorldPartitionHLODDetailsCustomization::CanBuildHLOD)
+			]
+			+ SHorizontalBox::Slot()
+			[
+				SNew(SButton)
+				.HAlign(EHorizontalAlignment::HAlign_Center)
+				.VAlign(EVerticalAlignment::VAlign_Center)
+				.Text(LOCTEXT("ForceBuildHLODButtonText", "Force Build HLOD"))
+				.OnClicked(this, &FWorldPartitionHLODDetailsCustomization::OnForceBuildHLOD)
+				.IsEnabled(this, &FWorldPartitionHLODDetailsCustomization::CanBuildHLOD)
+			]
 		]
 		+ SVerticalBox::Slot()
 		[
@@ -103,6 +118,16 @@ bool FWorldPartitionHLODDetailsCustomization::CanBuildHLOD() const
 }
 
 FReply FWorldPartitionHLODDetailsCustomization::OnBuildHLOD()
+{
+	return OnBuildHLODInternal(/*bForceBuild=*/false);
+}
+
+FReply FWorldPartitionHLODDetailsCustomization::OnForceBuildHLOD()
+{
+	return OnBuildHLODInternal(/*bForceBuild=*/true);
+}
+
+FReply FWorldPartitionHLODDetailsCustomization::OnBuildHLODInternal(bool bForceBuild)
 {
 	TArray<AWorldPartitionHLOD*> SelectedHLODActors = GetSelectedHLODActors();
 
@@ -137,19 +162,33 @@ FReply FWorldPartitionHLODDetailsCustomization::OnBuildHLOD()
 	FScopedSlowTask Progress(SelectedHLODActors.Num(), LOCTEXT("BuildingHLODsforActors", "Building HLODs for actors..."));
 	Progress.MakeDialog(true);
 
+	int32 NumHLODActorsUpdated = 0;
+
 	// Build HLODs
 	for (int32 Index = 0; Index < SelectedHLODActors.Num(); ++Index)
 	{
-		Progress.EnterProgressFrame(1.0f, FText::Format(LOCTEXT("BuildingHLODsforActorProgress", "Building HLODs ({0}/{1})"), Index+1, SelectedHLODActors.Num()));
+		Progress.EnterProgressFrame(1.0f, FText::Format(LOCTEXT("BuildingHLODsforActorProgress", "Building HLODs ({0}/{1})"), Index + 1, SelectedHLODActors.Num()));
 		if (Progress.ShouldCancel())
 		{
 			break;
 		}
 
 		AWorldPartitionHLOD* HLODActor = SelectedHLODActors[Index];
-		HLODActor->BuildHLOD(true);
+		bool bUpdated = HLODActor->BuildHLOD(bForceBuild);
+		NumHLODActorsUpdated += bUpdated ? 1 : 0;
 	};
-	
+
+	// Toast
+	FNotificationInfo Notification(LOCTEXT("HLODBuildCompleted", "HLOD Build Completed"));
+	Notification.Image = FAppStyle::GetBrush(TEXT("Icons.Info"));
+	Notification.bFireAndForget = true;
+	Notification.ExpireDuration = 3.0f;
+	Notification.bUseThrobber = true;
+
+	FText NumHLODActorsUpdatedText = NumHLODActorsUpdated > 0 ? FText::AsNumber(NumHLODActorsUpdated) : LOCTEXT("NoHLODActorUpdated", "No");
+	Notification.SubText = FText::Format(LOCTEXT("NumHLODActorsUpdated", "{0} HLOD {1}|plural(one=actor,other=actors) updated"), NumHLODActorsUpdatedText, NumHLODActorsUpdated);
+	FSlateNotificationManager::Get().AddNotification(Notification);
+
 	// Force refresh the UI so that any change to the selected actors are properly shown to the user
 	GEditor->NoteSelectionChange();
 
diff --git a/Engine/Source/Editor/WorldPartitionEditor/Private/WorldPartition/Customizations/WorldPartitionHLODDetailsCustomization.h b/Engine/Source/Editor/WorldPartitionEditor/Private/WorldPartition/Customizations/WorldPartitionHLODDetailsCustomization.h
--- a/Engine/Source/Editor/WorldPartitionEditor/Private/WorldPartition/Customizations/WorldPartitionHLODDetailsCustomization.h
+++ b/Engine/Source/Editor/WorldPartitionEditor/Private/WorldPartition/Customizations/WorldPartitionHLODDetailsCustomization.h
@@ -27,6 +27,8 @@ private:
 
 	bool CanBuildHLOD() const;
 	FReply OnBuildHLOD();
+	FReply OnForceBuildHLOD();
+	FReply OnBuildHLODInternal(bool bForceBuild);
 
 	static bool CanExportHLODAssets(const AWorldPartitionHLOD* HLODActor);
 	bool CanExportAssets() const;
diff --git a/Engine/Source/Runtime/Engine/Private/WorldPartition/HLOD/HLODActor.cpp b/Engine/Source/Runtime/Engine/Private/WorldPartition/HLOD/HLODActor.cpp
--- a/Engine/Source/Runtime/Engine/Private/WorldPartition/HLOD/HLODActor.cpp
+++ b/Engine/Source/Runtime/Engine/Private/WorldPartition/HLOD/HLODActor.cpp
@@ -70,7 +70,6 @@ AWorldPartitionHLOD::AWorldPartitionHLOD(const FObjectInitializer& ObjectInitial
 	SetNetUpdateFrequency(1.f);
 
 #if WITH_EDITORONLY_DATA
-	HLODHash = 0;
 	HLODBounds = FBox(EForceInit::ForceInit);
 #endif
 
@@ -578,30 +577,19 @@ int64 AWorldPartitionHLOD::GetStat(FName InStatName) const
 	return HLODStats.FindRef(InStatName);
 }
 
-uint32 AWorldPartitionHLOD::GetHLODHash() const
+bool AWorldPartitionHLOD::BuildHLOD(bool bForceBuild)
 {
-	return HLODHash;
-}
-
-void AWorldPartitionHLOD::SetHLODHash(uint32 InHLODHash, const FString& InHLODBuildReportContent)
-{
-	HLODHash = InHLODHash;
-	UpdateHLODBuildReportContent(InHLODBuildReportContent);
-}
+	bool bBuiltHLOD = false;
 
-void AWorldPartitionHLOD::BuildHLOD(bool bForceBuild)
-{
 	GetHLODBuildEventDelegate().Broadcast({ FWorldPartitionHLODBuildEvent::EEventType::BeginBuild, this });
 
 	IWorldPartitionHLODUtilitiesModule* WPHLODUtilitiesModule = FModuleManager::Get().LoadModulePtr<IWorldPartitionHLODUtilitiesModule>("WorldPartitionHLODUtilities");
 	if (IWorldPartitionHLODUtilities* WPHLODUtilities = WPHLODUtilitiesModule != nullptr ? WPHLODUtilitiesModule->GetUtilities() : nullptr)
 	{
-		if (bForceBuild)
-		{
-			HLODHash = 0;
-		}
-
-		WPHLODUtilities->BuildHLOD(this);
+		FHLODBuildParams BuildParams;
+		BuildParams.HLODActor = this;
+		BuildParams.bForceBuild = bForceBuild;
+		bBuiltHLOD = WPHLODUtilities->BuildHLOD(BuildParams);
 	}
 
 	// When generating WorldPartition HLODs, we have the renderer initialized.
@@ -611,17 +599,25 @@ void AWorldPartitionHLOD::BuildHLOD(bool bForceBuild)
 	BuildActorTextureStreamingData(this, EMaterialQualityLevel::High, GMaxRHIFeatureLevel);
 
 	GetHLODBuildEventDelegate().Broadcast({ FWorldPartitionHLODBuildEvent::EEventType::EndBuild, this });
+
+	return bBuiltHLOD;
 }
 
-uint32 AWorldPartitionHLOD::ComputeHLODHash() const
+bool AWorldPartitionHLOD::ShouldBuildHLOD() const
 {
-	uint32 ComputedHLODHash = 0;
+	bool bDirty = false;
+
 	IWorldPartitionHLODUtilitiesModule* WPHLODUtilitiesModule = FModuleManager::Get().LoadModulePtr<IWorldPartitionHLODUtilitiesModule>("WorldPartitionHLODUtilities");
 	if (IWorldPartitionHLODUtilities* WPHLODUtilities = WPHLODUtilitiesModule != nullptr ? WPHLODUtilitiesModule->GetUtilities() : nullptr)
 	{
-		ComputedHLODHash = WPHLODUtilities->ComputeHLODHash(this);
+		FHLODBuildParams BuildParams;
+		BuildParams.HLODActor = const_cast<AWorldPartitionHLOD*>(this);
+		BuildParams.bForceBuild = false;
+		BuildParams.bTestOnly = true;
+		bDirty = WPHLODUtilities->BuildHLOD(BuildParams);
 	}
-	return ComputedHLODHash;
+
+	return bDirty;
 }
 
 FWorldPartitionHLODBuildEventDelegate& AWorldPartitionHLOD::GetHLODBuildEventDelegate()
@@ -629,12 +625,6 @@ FWorldPartitionHLODBuildEventDelegate& AWorldPartitionHLOD::GetHLODBuildEventDel
 	return HLODBuildEventDelegate;
 }
 
-void AWorldPartitionHLOD::GetAssetRegistryTags(FAssetRegistryTagsContext Context) const
-{
-	Super::GetAssetRegistryTags(Context);
-	Context.AddTag(UObject::FAssetRegistryTag(NAME_HLODHash_AssetTag, LexToString(HLODHash), UObject::FAssetRegistryTag::TT_Hidden));
-}
-
 const TCHAR* HLOD_REPORT_BEGIN = TEXT("### HLOD_REPORT_BEGIN ###");
 const TCHAR* HLODLastBuildInfoSectionHeader = TEXT("### Last Build Info ###");
 const TCHAR* HLODActorDetailsSectionHeader = TEXT("### HLOD Actor Details ###");
@@ -675,7 +665,6 @@ FString AWorldPartitionHLOD::GenerateHLODBuildReportHeaderString() const
 
 	BuildReportHeader << HLODActorDetailsSectionHeader << LINE_TERMINATOR;
 	BuildReportHeader << TEXT(" * HLOD Actor Descriptor:   ") << LINE_TERMINATOR << TEXT("\t") << ActorDesc << LINE_TERMINATOR; // Will log on multiple lines
-	BuildReportHeader << TEXT(" * HLOD Actor Build Hash:   ") << FString::Printf(TEXT("%08X"), HLODHash) << LINE_TERMINATOR;
 	BuildReportHeader << TEXT(" * HLOD Actor IsStandalone: ") << (IsStandalone() ? TEXT("1") : TEXT("0")) << LINE_TERMINATOR;
 
 	return FString(BuildReportHeader);
@@ -728,7 +717,7 @@ void AWorldPartitionHLOD::UpdateHLODBuildReportHeader()
 	HLODBuildReport += BuildReportContent;
 }
 
-void AWorldPartitionHLOD::UpdateHLODBuildReportContent(const FString& InHLODBuildReportContent)
+void AWorldPartitionHLOD::UpdateHLODBuildReportContent()
 {
 	HLODBuildReport.Reset();
 
@@ -739,8 +728,17 @@ void AWorldPartitionHLOD::UpdateHLODBuildReportContent(const FString& InHLODBuil
 	HLODBuildReport += LINE_TERMINATOR;
 	HLODBuildReport += HLODBuildDetailsSectionHeader;
 	HLODBuildReport += LINE_TERMINATOR;
-	HLODBuildReport += InHLODBuildReportContent.TrimEnd();
-	HLODBuildReport += LINE_TERMINATOR;
+
+	for (const TObjectPtr<UHLODRebuildPolicyData> HLODRebuildPolicyData : HLODRebuildPolicyDataSet)
+	{
+		FString HLODPolicyDataString = HLODRebuildPolicyData->ToString();
+		if (!HLODPolicyDataString.IsEmpty())
+		{
+			HLODBuildReport += HLODRebuildPolicyData->ToString().TrimEnd();
+			HLODBuildReport += LINE_TERMINATOR;
+		}
+	}
+
 	HLODBuildReport += HLOD_REPORT_END;
 	HLODBuildReport += LINE_TERMINATOR;
 }
@@ -985,4 +983,21 @@ TArray<UObject*> AWorldPartitionHLOD::ExportHLODAssets(const FExportHLODAssetsPa
 	return ExportedAssets;
 }
 
+const FHLODRebuildPolicyDataSet& AWorldPartitionHLOD::GetHLODRebuildPolicyDataSet() const
+{
+	return HLODRebuildPolicyDataSet;
+}
+
+void AWorldPartitionHLOD::SetHLODRebuildPolicyDataSet(const FHLODRebuildPolicyDataSet& InHLODRebuildPolicyDataSet)
+{
+	HLODRebuildPolicyDataSet = InHLODRebuildPolicyDataSet;
+
+	for (UHLODRebuildPolicyData* HLODRebuildPolicyData : HLODRebuildPolicyDataSet)
+	{
+		HLODRebuildPolicyData->Rename(nullptr, this, REN_None);
+	}
+
+	UpdateHLODBuildReportContent();
+}
+
 #endif // WITH_EDITOR
\ No newline at end of file
diff --git a/Engine/Source/Runtime/Engine/Private/WorldPartition/HLOD/HLODRebuildPolicy.cpp b/Engine/Source/Runtime/Engine/Private/WorldPartition/HLOD/HLODRebuildPolicy.cpp
new file mode 100644
--- /dev/null
+++ b/Engine/Source/Runtime/Engine/Private/WorldPartition/HLOD/HLODRebuildPolicy.cpp
@@ -0,0 +1,134 @@
+// Copyright Epic Games, Inc. All Rights Reserved.
+
+#pragma once
+
+#include "WorldPartition/HLOD/HLODRebuildPolicy.h"
+#include "WorldPartition/HLOD/HLODActor.h"
+#include "Algo/RemoveIf.h"
+
+UHLODRebuildPolicyData::UHLODRebuildPolicyData(const FObjectInitializer& ObjectInitializer)
+	: Super(ObjectInitializer)
+{
+}
+
+UHLODRebuildPolicy::UHLODRebuildPolicy(const FObjectInitializer& ObjectInitializer)
+	: Super(ObjectInitializer)
+{
+}
+
+#if WITH_EDITOR
+
+TArray<TSubclassOf<UHLODRebuildPolicy>> UHLODRebuildPolicy::GetRebuildPoliciesClasses()
+{
+	TArray<TSubclassOf<UHLODRebuildPolicy>> HLODRebuildPoliciesClasses;
+
+	auto AddPolicyClass = [&HLODRebuildPoliciesClasses](const FSoftObjectPath& ObjectPath)
+	{
+		UObject* FoundObject = ObjectPath.TryLoad(nullptr, LOAD_Quiet);
+		if (FoundObject)
+		{
+			TSubclassOf<UHLODRebuildPolicy> PolicyClass = Cast<UClass>(FoundObject);
+			if (PolicyClass)
+			{
+				HLODRebuildPoliciesClasses.AddUnique(PolicyClass);
+			}
+			else
+			{
+				UE_LOG(LogHLODBuilder, Warning, TEXT("Invalid rebuild policy class specified: %s"), *ObjectPath.ToString());
+			}
+		}
+		else
+		{
+			UE_LOG(LogHLODBuilder, Warning, TEXT("Rebuild policy class not found: %s"), *ObjectPath.ToString());
+		}
+	};
+
+	for (const FSoftObjectPath& HLODRebuildPolicyClassPath : UHLODRebuildPolicy::StaticClass()->GetDefaultObject<UHLODRebuildPolicy>()->HLODRebuildPolicies)
+	{
+		AddPolicyClass(HLODRebuildPolicyClassPath);
+	}
+
+	if (HLODRebuildPoliciesClasses.IsEmpty())
+	{
+		AddPolicyClass(FSoftObjectPath(TEXT("/Script/WorldPartitionHLODUtilities.HLODRebuildPolicyHashCompare")));
+	}
+
+	return HLODRebuildPoliciesClasses;
+}
+
+FHLODRebuildPolicyDataSet UHLODRebuildPolicy::ComputeDataForRebuildPolicies(const AWorldPartitionHLOD* InHLODActor, const TArray<UActorComponent*>& InSourceComponents, const bool bInForComparison)
+{
+	FHLODRebuildPolicyDataSet HLODRebuildPolicyDataSet;
+
+	for (TSubclassOf<UHLODRebuildPolicy> HLODRebuildPolicyClass : GetRebuildPoliciesClasses())
+	{
+		UHLODRebuildPolicy* UHLODRebuildPolicyCDO = HLODRebuildPolicyClass.GetDefaultObject();
+		check(UHLODRebuildPolicyCDO);
+
+		UHLODRebuildPolicyData* OldData = nullptr;
+		if (bInForComparison)
+		{
+			const FHLODRebuildPolicyDataSet& BaselineHLODRebuildPolicyDataSet = InHLODActor->GetHLODRebuildPolicyDataSet();
+			const TObjectPtr<UHLODRebuildPolicyData>* OldDataPtr = BaselineHLODRebuildPolicyDataSet.FindByPredicate([UHLODRebuildPolicyCDO](const TObjectPtr<UHLODRebuildPolicyData>& Entry) { return Entry.IsA(UHLODRebuildPolicyCDO->GetRebuildPolicyDataType()); });
+			OldData = OldDataPtr ? *OldDataPtr : nullptr;
+		}
+		
+		UHLODRebuildPolicyData* NewData = UHLODRebuildPolicyCDO->ComputeDataForRebuildPolicy(InHLODActor, InSourceComponents, OldData);
+		if (NewData)
+		{
+			HLODRebuildPolicyDataSet.Emplace(NewData);
+		}
+	}
+
+	return HLODRebuildPolicyDataSet;
+}
+
+EHLODRebuildPolicyDecision UHLODRebuildPolicy::Evaluate(const AWorldPartitionHLOD* InHLODActor, const FHLODRebuildPolicyDataSet& InOldData, const FHLODRebuildPolicyDataSet& InNewData)
+{
+	EHLODRebuildPolicyDecision FinalDecision = EHLODRebuildPolicyDecision::None;
+
+	UE_LOG(LogHLODBuilder, Log, TEXT("Evaluating HLOD rebuild policies..."));
+
+	for (TSubclassOf<UHLODRebuildPolicy> HLODRebuildPolicyClass : GetRebuildPoliciesClasses())
+	{
+		UHLODRebuildPolicy* UHLODRebuildPolicyCDO = HLODRebuildPolicyClass.GetDefaultObject();
+		check(UHLODRebuildPolicyCDO);
+
+		// Find old data
+		const TObjectPtr<UHLODRebuildPolicyData>* OldHLODRebuildPolicyDataPtr = InOldData.FindByPredicate([UHLODRebuildPolicyCDO](const TObjectPtr<UHLODRebuildPolicyData>& Entry) { return Entry.IsA(UHLODRebuildPolicyCDO->GetRebuildPolicyDataType()); });
+		const UHLODRebuildPolicyData* OldHLODRebuildPolicyData = OldHLODRebuildPolicyDataPtr ? *OldHLODRebuildPolicyDataPtr : nullptr;
+
+		// Find new Data
+		const TObjectPtr<UHLODRebuildPolicyData>* NewHLODRebuildPolicyDataPtr = InNewData.FindByPredicate([UHLODRebuildPolicyCDO](const TObjectPtr<UHLODRebuildPolicyData>& Entry) { return Entry.IsA(UHLODRebuildPolicyCDO->GetRebuildPolicyDataType()); });
+		const UHLODRebuildPolicyData* NewHLODRebuildPolicyData = NewHLODRebuildPolicyDataPtr ? *NewHLODRebuildPolicyDataPtr : nullptr;
+
+		if (!NewHLODRebuildPolicyData)
+		{
+			continue;
+		}
+		
+		FString Reason;
+		EHLODRebuildPolicyDecision Decision = UHLODRebuildPolicyCDO->Evaluate(InHLODActor, OldHLODRebuildPolicyData, NewHLODRebuildPolicyData, Reason);
+
+		FString PolicyName = UHLODRebuildPolicyCDO->GetClass()->GetName();
+		FString DecisionString = StaticEnum<EHLODRebuildPolicyDecision>()->GetNameStringByValue(static_cast<int64>(Decision));
+		
+		UE_CLOG( Reason.IsEmpty(), LogHLODBuilder, Log, TEXT(" * %s -> %s"), *PolicyName, *DecisionString);
+		UE_CLOG(!Reason.IsEmpty(), LogHLODBuilder, Log, TEXT(" * %s -> %s (%s)"), *PolicyName, *DecisionString, *Reason);
+		
+		if (Decision != EHLODRebuildPolicyDecision::None)
+		{
+			FinalDecision = Decision;
+		}
+
+		// If the build was rejected, stop processing
+		if (FinalDecision == EHLODRebuildPolicyDecision::RejectRebuild)
+		{
+			break;
+		}
+	}
+
+	return FinalDecision;
+}
+
+#endif
\ No newline at end of file
diff --git a/Engine/Source/Runtime/Engine/Public/WorldPartition/HLOD/HLODActor.h b/Engine/Source/Runtime/Engine/Public/WorldPartition/HLOD/HLODActor.h
--- a/Engine/Source/Runtime/Engine/Public/WorldPartition/HLOD/HLODActor.h
+++ b/Engine/Source/Runtime/Engine/Public/WorldPartition/HLOD/HLODActor.h
@@ -9,10 +9,10 @@
 #include "WorldPartition/WorldPartitionActorDesc.h"
 #include "WorldPartition/HLOD/HLODStats.h"
 #include "WorldPartition/HLOD/HLODBuilder.h"
+#include "WorldPartition/HLOD/HLODRebuildPolicy.h"
 #include "WorldPartition/HLOD/IWorldPartitionHLODObject.h"
 
 #if WITH_EDITOR
-#include "UObject/AssetRegistryTagsContext.h"
 #include "UObject/ObjectSaveContext.h"
 #endif // WITH_EDITOR
 
@@ -20,6 +20,7 @@
 
 class UHLODLayer;
 class UWorldPartitionHLODSourceActors;
+class UHLODRebuildPolicyData;
 
 ENGINE_API DECLARE_LOG_CATEGORY_EXTERN(LogHLODHash, Log, All);
 
@@ -115,10 +116,8 @@ public:
 	double GetMinVisibleDistance() const { return MinVisibleDistance; }
 	void SetMinVisibleDistance(double InMinVisibleDistance) { MinVisibleDistance = InMinVisibleDistance; }
 
-	ENGINE_API void BuildHLOD(bool bForceBuild = false);
-	ENGINE_API uint32 ComputeHLODHash() const;
-	ENGINE_API uint32 GetHLODHash() const;
-	ENGINE_API void SetHLODHash(uint32 InHLODHash, const FString& InHLODBuildReportContent);
+	ENGINE_API bool BuildHLOD(bool bForceBuild = false);
+	ENGINE_API bool ShouldBuildHLOD() const;
 
 	ENGINE_API int64 GetStat(FName InStatName) const;
 	void SetStat(FName InStatName, int64 InStatValue) { HLODStats.Add(InStatName, InStatValue); }
@@ -128,8 +127,20 @@ public:
 
 	UFUNCTION(BlueprintCallable, Category = "World Partition", meta = (DevelopmentOnly))
 	ENGINE_API TArray<UObject*> ExportHLODAssets(const struct FExportHLODAssetsParams& ExportHLODAssetsParams, FString& OutErrorMessage) const;
-
+	
 	ENGINE_API void UpdateHLODBuildReportHeader();
+	ENGINE_API void UpdateHLODBuildReportContent();
+
+	ENGINE_API const FHLODRebuildPolicyDataSet& GetHLODRebuildPolicyDataSet() const;
+	ENGINE_API void SetHLODRebuildPolicyDataSet(const FHLODRebuildPolicyDataSet& InHLODRebuildPolicyDataSet);
+
+	UE_DEPRECATED(5.8, "Hash data is now stored in a UHLODRebuildPolicyData object.")
+	ENGINE_API uint32 GetHLODHash() const { return HLODHash_DEPRECATED; }
+	UE_DEPRECATED(5.8, "Hash data is now stored in a UHLODRebuildPolicyData object.")
+	ENGINE_API uint32 ComputeHLODHash() const { return 0; }
+	UE_DEPRECATED(5.8, "Hash data is now stored in a UHLODRebuildPolicyData object.")
+	ENGINE_API void SetHLODHash(uint32 InHLODHash, const FString& InHLODBuildReportContent) {}
+	UE_DEPRECATED(5.8, "Use override with no parameter.")
 	ENGINE_API void UpdateHLODBuildReportContent(const FString& InHLODBuildReportContent);
 
 private:
@@ -176,7 +187,6 @@ protected:
 	virtual bool IsLockLocation() const override { return true; }
 	virtual bool IsUserManaged() const override { return false; }
 	virtual bool IsActorLabelEditable() const override { return false; }
-	ENGINE_API virtual void GetAssetRegistryTags(FAssetRegistryTagsContext Context) const override;
 #endif
 	//~ End AActor Interface.
 
@@ -199,14 +209,14 @@ private:
 	UPROPERTY()
 	double MinVisibleDistance;
 
-	UPROPERTY()
-	uint32 HLODHash;
-
 	UPROPERTY()
 	TMap<FName, int64> HLODStats;
 
 	UPROPERTY()
 	FString HLODBuildReport;
+
+	UPROPERTY(DuplicateTransient)
+	TArray<TObjectPtr<UHLODRebuildPolicyData>> HLODRebuildPolicyDataSet;
 #endif
 
 	UPROPERTY()
@@ -233,6 +243,9 @@ private:
 
 	UPROPERTY()
 	TObjectPtr<const UHLODLayer> SubActorsHLODLayer_DEPRECATED;
+
+	UPROPERTY()
+	uint32 HLODHash_DEPRECATED;
 #endif
 
 #if WITH_EDITOR
diff --git a/Engine/Source/Runtime/Engine/Public/WorldPartition/HLOD/HLODRebuildPolicy.h b/Engine/Source/Runtime/Engine/Public/WorldPartition/HLOD/HLODRebuildPolicy.h
new file mode 100644
--- /dev/null
+++ b/Engine/Source/Runtime/Engine/Public/WorldPartition/HLOD/HLODRebuildPolicy.h
@@ -0,0 +1,61 @@
+// Copyright Epic Games, Inc. All Rights Reserved.
+
+#pragma once
+
+#include "UObject/Object.h"
+#include "Templates/SubclassOf.h"
+#include "HLODRebuildPolicy.generated.h"
+
+class AWorldPartitionHLOD;
+class UActorComponent;
+
+/**
+ * Decision type returned when a rebuild policy is evaluated.
+ */
+UENUM()
+enum class EHLODRebuildPolicyDecision : uint8
+{
+	None,			//!< No decision.
+	ApproveRebuild,	//!< At least one HLOD rebuild policy must approve the build for it to be performed.
+	RejectRebuild	//!< If a single rebuild policy rejects the build, it is aborted.
+};
+
+/**
+ * Data associated with the policy evaluation.
+ * Stored in the HLOD actor so that it can be used as the baseline when performing the next evaluation.
+ */
+UCLASS(Abstract, MinimalAPI)
+class UHLODRebuildPolicyData : public UObject
+{
+	GENERATED_UCLASS_BODY()
+
+public:
+	virtual FString ToString() const { return FString(); };
+};
+
+typedef TArray<TObjectPtr<UHLODRebuildPolicyData>> FHLODRebuildPolicyDataSet;
+
+/**
+ * Logic and settings associated with a rebuild policy.
+  */
+UCLASS(Abstract, MinimalAPI, Config = Editor)
+class UHLODRebuildPolicy : public UObject
+{
+	GENERATED_UCLASS_BODY()
+
+#if WITH_EDITOR
+	static ENGINE_API FHLODRebuildPolicyDataSet ComputeDataForRebuildPolicies(const AWorldPartitionHLOD* InHLODActor, const TArray<UActorComponent*>& InSourceComponents, const bool bInForComparison);
+	static ENGINE_API EHLODRebuildPolicyDecision Evaluate(const AWorldPartitionHLOD* InHLODActor, const FHLODRebuildPolicyDataSet& InOldData, const FHLODRebuildPolicyDataSet& InNewData);
+
+protected:
+	virtual TSubclassOf<UHLODRebuildPolicyData> GetRebuildPolicyDataType() const PURE_VIRTUAL(UHLODRebuildPolicy::GetRebuildPolicyDataType, return UHLODRebuildPolicyData::StaticClass(); );
+	virtual EHLODRebuildPolicyDecision Evaluate(const AWorldPartitionHLOD* InHLODActor, const UHLODRebuildPolicyData* InOldData, const UHLODRebuildPolicyData* InNewData, FString& OutReason) const PURE_VIRTUAL(UHLODRebuildPolicy::Evaluate, return EHLODRebuildPolicyDecision::None; );
+	virtual UHLODRebuildPolicyData* ComputeDataForRebuildPolicy(const AWorldPartitionHLOD* InHLODActor, const TArray<UActorComponent*>& InSourceComponents, const UHLODRebuildPolicyData* InExistingPolicyData) const PURE_VIRTUAL(UHLODRebuildPolicy::ComputeDataForRebuildPolicy, return nullptr; );
+#endif
+
+	static TArray<TSubclassOf<UHLODRebuildPolicy>> GetRebuildPoliciesClasses();
+
+	// HLOD Rebuild Policies to evaluate. These are evaluated one after the other, so the order is important.
+	UPROPERTY(Config)
+	TArray<FSoftObjectPath> HLODRebuildPolicies;
+};
diff --git a/Engine/Source/Runtime/Engine/Public/WorldPartition/HLOD/IWorldPartitionHLODUtilities.h b/Engine/Source/Runtime/Engine/Public/WorldPartition/HLOD/IWorldPartitionHLODUtilities.h
--- a/Engine/Source/Runtime/Engine/Public/WorldPartition/HLOD/IWorldPartitionHLODUtilities.h
+++ b/Engine/Source/Runtime/Engine/Public/WorldPartition/HLOD/IWorldPartitionHLODUtilities.h
@@ -10,6 +10,7 @@
 
 #if WITH_EDITOR
 
+#include "WorldPartition/HLOD/HLODRebuildPolicy.h"
 #include "WorldPartition/DataLayer/ExternalDataLayerAsset.h"
 #include "WorldPartition/DataLayer/ExternalDataLayerInstance.h"
 
@@ -59,6 +60,13 @@ struct FHLODCreationParams
 	FBox CellBounds;
 };
 
+struct FHLODBuildParams
+{
+	AWorldPartitionHLOD* HLODActor = nullptr;
+	bool bForceBuild = false;
+	bool bTestOnly = false;
+};
+
 /**
  * Tools for building HLODs in WorldPartition
  */
@@ -73,25 +81,16 @@ public:
 	 * @param	InCreationContext	HLOD creation context object
 	 * @param	InCreationParams	HLOD creation parameters object
 	 * @param	InActors			The actors for which we'll build an HLOD representation
-	 * @param	InDataLayers		The data layers to assign to the newly created HLOD actors
 	 */
 	virtual TArray<AWorldPartitionHLOD*> CreateHLODActors(FHLODCreationContext& InCreationContext, const FHLODCreationParams& InCreationParams, const TArray<IStreamingGenerationContext::FActorInstance>& InActors) = 0;
 
 	/**
 	 * Build HLOD for the specified AWorldPartitionHLOD actor.
 	 *
-	 * @param 	InHLODActor		The HLOD actor for which we'll build the HLOD
-	 * @return An hash that represent the content used to build this HLOD.
+	 * @param 	InBuildParams		Parameters to use for the HLOD build
+	 * @return True if the HLOD actor was updated, false otherwise.
 	 */
-	virtual uint32 BuildHLOD(AWorldPartitionHLOD* InHLODActor) = 0;
-
-	/**
-	 * Compute the HLOD hash for the specified AWorldPartitionHLOD actor.
-	 *
-	 * @param 	InHLODActor		The HLOD actor for which we'll compute the hash
-	 * @return An hash that represent the content used to build this HLOD.
-	 */
-	virtual uint32 ComputeHLODHash(const AWorldPartitionHLOD* InHLODActor) = 0;
+	virtual bool BuildHLOD(const FHLODBuildParams& InBuildParams) = 0;
 
 	/**
 	 * Retrieve the HLOD Builder class to use for the given HLODLayer.
@@ -112,26 +111,32 @@ public:
 	/**
 	 * HLOD build evaluator delegate
 	 * @param HLODActor			The HLOD actor to be rebuilt.
-	 * @param OldHash			The previously stored hash of the inputs to the HLOD build.
-	 * @param NewHash			The newly computed hash of the inputs to the HLOD build.
-	 * @return true if the HLOD build should be performed, false otherwise.
+	 * @param OldPolicyDataSet	The previously stored policy data set of the inputs to the HLOD build.
+	 * @param NewPolicyDataSet	The newly computed policy data set of the inputs to the HLOD build.
+	 * @return The HLOD rebuild decision
 	 */
-	UE_EXPERIMENTAL(5.7, "Experimental API, HLOD build evaluator logic may be moved in a future release.")
-	DECLARE_DELEGATE_RetVal_ThreeParams(bool, FHLODBuildEvaluator, AWorldPartitionHLOD*, uint32, uint32);
+	DECLARE_DELEGATE_RetVal_ThreeParams(EHLODRebuildPolicyDecision, FHLODBuildEvaluator, AWorldPartitionHLOD*, const FHLODRebuildPolicyDataSet&, const FHLODRebuildPolicyDataSet&);
 
 	/**
 	 * Provide a delegate that will be used to evaluate whether an HLOD build should be performed.
 	 * @param The evaluator delegate.
 	 */
-	UE_EXPERIMENTAL(5.7, "Experimental API, HLOD build evaluator logic may be moved in a future release.")
 	virtual void SetHLODBuildEvaluator(FHLODBuildEvaluator BuildEvaluatorDelegate) = 0;
 
-
 	UE_DEPRECATED(5.2, "Use the overload that passes the DataLayersInstances via InCreationParams")
 	virtual TArray<AWorldPartitionHLOD*> CreateHLODActors(FHLODCreationContext& InCreationContext, const FHLODCreationParams& InCreationParams, const TArray<IStreamingGenerationContext::FActorInstance>& InActors, const TArray<const UDataLayerInstance*>& InDataLayerInstances)
 	{
 		const_cast<FHLODCreationParams&>(InCreationParams).DataLayerInstances = InDataLayerInstances;
 		return CreateHLODActors(InCreationContext, InCreationParams, InActors);
 	}
+
+	UE_DEPRECATED(5.8, "Use the overload that passes a FHLODBuildParams")
+	virtual bool BuildHLOD(AWorldPartitionHLOD* InHLODActor) final
+	{
+		FHLODBuildParams Params;
+		Params.HLODActor = InHLODActor;
+		Params.bForceBuild = false;
+		return BuildHLOD(Params);
+	}
 };
 #endif
