CIS fix for issue introduced when resolving robomerge conflict for 48769649

#Synced-CL 48765230

diff --git a/Engine/Plugins/Editor/WorldPartitionHLODUtilities/Source/Private/WorldPartition/HLOD/Utilities/WorldPartitionHLODUtilities.cpp b/Engine/Plugins/Editor/WorldPartitionHLODUtilities/Source/Private/WorldPartition/HLOD/Utilities/WorldPartitionHLODUtilities.cpp
--- a/Engine/Plugins/Editor/WorldPartitionHLODUtilities/Source/Private/WorldPartition/HLOD/Utilities/WorldPartitionHLODUtilities.cpp
+++ b/Engine/Plugins/Editor/WorldPartitionHLODUtilities/Source/Private/WorldPartition/HLOD/Utilities/WorldPartitionHLODUtilities.cpp
@@ -818,7 +818,7 @@ bool FWorldPartitionHLODUtilities::BuildHLOD(const FHLODBuildParams& InBuildPara
 
 	const bool bForceBuild = bIsDirty || InBuildParams.bForceBuild;
 
-	const UWorldPartitionHLODSourceActors* SourceActors = InHLODActor->GetSourceActors();
+	const UWorldPartitionHLODSourceActors* SourceActors = HLODActor->GetSourceActors();
 	check(SourceActors);
 	TArray<UActorComponent*> HLODRelevantComponents = GatherHLODRelevantComponents(TempWorld, *SourceActors);
 
