Fix static analysis warning

#Synced-CL 48731319

diff --git a/Engine/Plugins/Editor/WorldPartitionHLODUtilities/Source/Private/WorldPartition/HLOD/RebuildPolicies/HLODRebuildPolicyImageCompare.cpp b/Engine/Plugins/Editor/WorldPartitionHLODUtilities/Source/Private/WorldPartition/HLOD/RebuildPolicies/HLODRebuildPolicyImageCompare.cpp
--- a/Engine/Plugins/Editor/WorldPartitionHLODUtilities/Source/Private/WorldPartition/HLOD/RebuildPolicies/HLODRebuildPolicyImageCompare.cpp
+++ b/Engine/Plugins/Editor/WorldPartitionHLODUtilities/Source/Private/WorldPartition/HLOD/RebuildPolicies/HLODRebuildPolicyImageCompare.cpp
@@ -134,7 +134,7 @@ static void PerformSceneWarmup(UWorld* InCaptureWorld, const TArray<UActorCompon
 						}
 
 						UMaterialInstance* MaterialInstance = Cast<UMaterialInstance>(MaterialInterface);
-						MaterialInterface = MaterialInstance ? MaterialInstance->Parent : nullptr;
+						MaterialInterface = MaterialInstance ? MaterialInstance->Parent.Get() : nullptr;
 					};
 				}
 			}
