World Partition - HLOD: Revert back to the previous behavior where each HLOD build is performed using it's own temp package
* Fixes issue in back to back HLOD builds for multiple maps using a collection
* Will not lead to FName bloat (issue addressed in prior CL)

#Synced-CL 48784656

diff --git a/Engine/Plugins/Editor/WorldPartitionHLODUtilities/Source/Private/WorldPartition/HLOD/Utilities/WorldPartitionHLODUtilities.cpp b/Engine/Plugins/Editor/WorldPartitionHLODUtilities/Source/Private/WorldPartition/HLOD/Utilities/WorldPartitionHLODUtilities.cpp
--- a/Engine/Plugins/Editor/WorldPartitionHLODUtilities/Source/Private/WorldPartition/HLOD/Utilities/WorldPartitionHLODUtilities.cpp
+++ b/Engine/Plugins/Editor/WorldPartitionHLODUtilities/Source/Private/WorldPartition/HLOD/Utilities/WorldPartitionHLODUtilities.cpp
@@ -794,7 +794,8 @@ bool FWorldPartitionHLODUtilities::BuildHLOD(const FHLODBuildParams& InBuildPara
 
 	const bool bInformEngineOfWorld = false;
 	const bool bAddToRoot = true;
-	UWorld* TempWorld = UWorld::CreateWorld(EWorldType::Editor, bInformEngineOfWorld, MakeUniqueObjectName(GetTransientPackage(), UWorld::StaticClass(), TEXT("BuildHLODWorld")), GetTransientPackage(), bAddToRoot);
+	UPackage* TempPackage = NewObject<UPackage>(nullptr, *MakeUniqueObjectName(nullptr, UPackage::StaticClass(), TEXT("/Temp/BuildHLODPackage")).ToString(), RF_Transient);
+	UWorld* TempWorld = UWorld::CreateWorld(EWorldType::Editor, bInformEngineOfWorld, TEXT("BuildHLODWorld"), TempPackage, bAddToRoot);
 	check(TempWorld);
 
 	ON_SCOPE_EXIT
