World Partition - HLOD: Default initialize members of the FHLODCaptureImage and FHLODCaptureFrame USTRUCTs to silence errors at startup
#Synced-CL 48819460

diff --git a/Engine/Plugins/Editor/WorldPartitionHLODUtilities/Source/Public/WorldPartition/HLOD/RebuildPolicies/HLODRebuildPolicyImageCompare.h b/Engine/Plugins/Editor/WorldPartitionHLODUtilities/Source/Public/WorldPartition/HLOD/RebuildPolicies/HLODRebuildPolicyImageCompare.h
--- a/Engine/Plugins/Editor/WorldPartitionHLODUtilities/Source/Public/WorldPartition/HLOD/RebuildPolicies/HLODRebuildPolicyImageCompare.h
+++ b/Engine/Plugins/Editor/WorldPartitionHLODUtilities/Source/Public/WorldPartition/HLOD/RebuildPolicies/HLODRebuildPolicyImageCompare.h
@@ -29,6 +29,12 @@ struct FHLODCaptureImage
 {
 	GENERATED_BODY()
 
+	FHLODCaptureImage()
+	: Width(0)
+	, Height(0)
+	, BytesPerPixel(0)
+	{}
+
 	UPROPERTY()
 	int32 Width;
 
@@ -50,6 +56,13 @@ struct FHLODCaptureFrame
 {
 	GENERATED_BODY()
 
+	FHLODCaptureFrame()
+		: CameraLocation(ForceInitToZero)
+		, CameraRotation(ForceInitToZero)
+		, OrthoWidth(0.0f)
+	{
+	}
+
 	UPROPERTY()
 	FVector CameraLocation;
 
