World Partition - HLOD: Fix incorrect world used in UMaterialParameterCollectionHLODModifier
* Renamed FHLODBuildContext::World to TargetWorld to disambiguate against the world used to build the HLODs
#Synced-CL 48819460

diff --git a/Engine/Plugins/Editor/WorldPartitionHLODUtilities/Source/Private/WorldPartition/HLOD/Builders/HLODBuilderMeshMerge.cpp b/Engine/Plugins/Editor/WorldPartitionHLODUtilities/Source/Private/WorldPartition/HLOD/Builders/HLODBuilderMeshMerge.cpp
--- a/Engine/Plugins/Editor/WorldPartitionHLODUtilities/Source/Private/WorldPartition/HLOD/Builders/HLODBuilderMeshMerge.cpp
+++ b/Engine/Plugins/Editor/WorldPartitionHLODUtilities/Source/Private/WorldPartition/HLOD/Builders/HLODBuilderMeshMerge.cpp
@@ -99,7 +99,7 @@ TArray<UActorComponent*> UHLODBuilderMeshMerge::Build(const FHLODBuildContext& I
 	}
 
 	const IMeshMergeUtilities& MeshMergeUtilities = FModuleManager::Get().LoadModuleChecked<IMeshMergeModule>("MeshMergeUtilities").GetUtilities();
-	MeshMergeUtilities.MergeComponentsToStaticMesh(SourcePrimitiveComponents, InHLODBuildContext.World, UseSettings, HLODMaterial, InHLODBuildContext.AssetsOuter->GetPackage(), InHLODBuildContext.AssetsBaseName, Assets, MergedActorLocation, 0.25f, false);
+	MeshMergeUtilities.MergeComponentsToStaticMesh(SourcePrimitiveComponents, InHLODBuildContext.BuildWorld, UseSettings, HLODMaterial, InHLODBuildContext.AssetsOuter->GetPackage(), InHLODBuildContext.AssetsBaseName, Assets, MergedActorLocation, 0.25f, false);
 
 	UStaticMeshComponent* Component = nullptr;
 	Algo::ForEach(Assets, [this, &Component, &MergedActorLocation](UObject* Asset)
diff --git a/Engine/Plugins/Editor/WorldPartitionHLODUtilities/Source/Private/WorldPartition/HLOD/Utilities/WorldPartitionHLODUtilities.cpp b/Engine/Plugins/Editor/WorldPartitionHLODUtilities/Source/Private/WorldPartition/HLOD/Utilities/WorldPartitionHLODUtilities.cpp
--- a/Engine/Plugins/Editor/WorldPartitionHLODUtilities/Source/Private/WorldPartition/HLOD/Utilities/WorldPartitionHLODUtilities.cpp
+++ b/Engine/Plugins/Editor/WorldPartitionHLODUtilities/Source/Private/WorldPartition/HLOD/Utilities/WorldPartitionHLODUtilities.cpp
@@ -794,19 +794,19 @@ bool FWorldPartitionHLODUtilities::BuildHLOD(const FHLODBuildParams& InBuildPara
 
 	const bool bInformEngineOfWorld = false;
 	const bool bAddToRoot = true;
-	UPackage* TempPackage = NewObject<UPackage>(nullptr, *MakeUniqueObjectName(nullptr, UPackage::StaticClass(), TEXT("/Temp/BuildHLODPackage")).ToString(), RF_Transient);
-	UWorld* TempWorld = UWorld::CreateWorld(EWorldType::Editor, bInformEngineOfWorld, TEXT("BuildHLODWorld"), TempPackage, bAddToRoot);
-	check(TempWorld);
+	UPackage* BuildHLODPackage = NewObject<UPackage>(nullptr, *MakeUniqueObjectName(nullptr, UPackage::StaticClass(), TEXT("/Temp/BuildHLODPackage")).ToString(), RF_Transient);
+	UWorld* BuildHLODWorld = UWorld::CreateWorld(EWorldType::Editor, bInformEngineOfWorld, TEXT("BuildHLODWorld"), BuildHLODPackage, bAddToRoot);
+	check(BuildHLODWorld);
 
 	ON_SCOPE_EXIT
 	{
-		TempWorld->DestroyWorld(bInformEngineOfWorld);
+		BuildHLODWorld->DestroyWorld(bInformEngineOfWorld);
 	};
 
 	{
 		FAutoScopedDurationTimer LoadTimeScope;
 
-		bool bLoadSucceeded = LoadSourceActors(HLODActor, TempWorld, bIsDirty);
+		bool bLoadSucceeded = LoadSourceActors(HLODActor, BuildHLODWorld, bIsDirty);
 		if (!ensure(bLoadSucceeded))
 		{
 			UE_LOG(LogHLODBuilder, Error, TEXT("FWorldPartitionHLODUtilities::BuildHLOD() - Failed to load source actors"));
@@ -820,7 +820,7 @@ bool FWorldPartitionHLODUtilities::BuildHLOD(const FHLODBuildParams& InBuildPara
 
 	const UWorldPartitionHLODSourceActors* SourceActors = HLODActor->GetSourceActors();
 	check(SourceActors);
-	TArray<UActorComponent*> HLODRelevantComponents = GatherHLODRelevantComponents(TempWorld, *SourceActors);
+	TArray<UActorComponent*> HLODRelevantComponents = GatherHLODRelevantComponents(BuildHLODWorld, *SourceActors);
 
 	FHLODRebuildPolicyDataSet OldRebuildPolicyDataSet = bForceBuild ? FHLODRebuildPolicyDataSet() : HLODActor->GetHLODRebuildPolicyDataSet();
 	FHLODRebuildPolicyDataSet NewRebuildPolicyDataSet = UHLODRebuildPolicy::ComputeDataForRebuildPolicies(HLODActor, HLODRelevantComponents, /*bInForComparison=*/true);
@@ -888,12 +888,16 @@ bool FWorldPartitionHLODUtilities::BuildHLOD(const FHLODBuildParams& InBuildPara
 			HLODBuilder->SetHLODBuilderSettings(HLODLayer->GetHLODBuilderSettings());
 
 			FHLODBuildContext HLODBuildContext;
-			HLODBuildContext.World = HLODActor->GetWorld();
+			HLODBuildContext.TargetWorld = HLODActor->GetWorld();
+			HLODBuildContext.BuildWorld = BuildHLODWorld;
 			HLODBuildContext.SourceComponents = HLODRelevantComponents;
 			HLODBuildContext.AssetsOuter = HLODActor->GetPackage();
 			HLODBuildContext.AssetsBaseName = HLODActor->GetActorLabel();
 			HLODBuildContext.MinVisibleDistance = HLODActor->GetMinVisibleDistance();
 			HLODBuildContext.WorldPosition = HLODActor->GetHLODBounds().GetCenter();
+PRAGMA_DISABLE_DEPRECATION_WARNINGS
+			HLODBuildContext.World = HLODBuildContext.TargetWorld;
+PRAGMA_ENABLE_DEPRECATION_WARNINGS
 
 			const TSubclassOf<UWorldPartitionHLODModifier> HLODModifierClass = HLODLayer->GetHLODModifierClass();
 			UWorldPartitionHLODModifier* HLODModifier = HLODModifierClass.Get() ? NewObject<UWorldPartitionHLODModifier>(GetTransientPackage(), HLODModifierClass) : nullptr;
diff --git a/Engine/Source/Runtime/Engine/Private/WorldPartition/HLOD/MaterialParameterCollectionHLODModifier.cpp b/Engine/Source/Runtime/Engine/Private/WorldPartition/HLOD/MaterialParameterCollectionHLODModifier.cpp
--- a/Engine/Source/Runtime/Engine/Private/WorldPartition/HLOD/MaterialParameterCollectionHLODModifier.cpp
+++ b/Engine/Source/Runtime/Engine/Private/WorldPartition/HLOD/MaterialParameterCollectionHLODModifier.cpp
@@ -25,7 +25,7 @@ void UMaterialParameterCollectionHLODModifier::BeginHLODBuild(const FHLODBuildCo
 	HLODBuildContext = &InHLODBuildContext;
 	if (MPC != nullptr)
 	{
-		if (UMaterialParameterCollectionInstance* Instance = HLODBuildContext->World->GetParameterCollectionInstance(MPC))
+		if (UMaterialParameterCollectionInstance* Instance = HLODBuildContext->BuildWorld->GetParameterCollectionInstance(MPC))
 		{
 			ApplyScalarParameterOverrides(Instance);
 		}
@@ -36,7 +36,7 @@ void UMaterialParameterCollectionHLODModifier::EndHLODBuild(TArray<UActorCompone
 {
 	if (MPC != nullptr)
 	{
-		if (UMaterialParameterCollectionInstance* Instance = HLODBuildContext->World->GetParameterCollectionInstance(MPC))
+		if (UMaterialParameterCollectionInstance* Instance = HLODBuildContext->BuildWorld->GetParameterCollectionInstance(MPC))
 		{
 			RestoreScalarParameterValues(Instance);
 		}
diff --git a/Engine/Source/Runtime/Engine/Public/WorldPartition/HLOD/HLODBuilder.h b/Engine/Source/Runtime/Engine/Public/WorldPartition/HLOD/HLODBuilder.h
--- a/Engine/Source/Runtime/Engine/Public/WorldPartition/HLOD/HLODBuilder.h
+++ b/Engine/Source/Runtime/Engine/Public/WorldPartition/HLOD/HLODBuilder.h
@@ -51,9 +51,15 @@ public:
  */
 struct FHLODBuildContext
 {
-	/** World for which HLODs are being built */
+	UE_DEPRECATED(5.8, "Replaced by TargetWorld to disambiguate vs BuildWorld")
 	UWorld*	World;
 
+	/** World for which HLODs are being built */
+	UWorld* TargetWorld;
+
+	/** Temporary world used to perform the HLOD build */
+	UWorld* BuildWorld;
+
 	/** Components that will be represented by this HLOD */
 	TArray<UActorComponent*> SourceComponents;
 
