[StandaloneHLOD][Experimental] Add an option to externalize HLOD assets
* When enabled, HLOD assets will be stored in a separate package instead of being embedded in the HLOD Actor's package.
* This prevents data duplication when using the same Level with Standalone HLOD as a LevelInstance or HLODOnlyLevelInstance in multiple other levels.

#Synced-CL 49190024

diff --git a/Engine/Plugins/Editor/WorldPartitionHLODUtilities/Source/Private/WorldPartition/HLOD/Builders/HLODBuilderMeshApproximate.cpp b/Engine/Plugins/Editor/WorldPartitionHLODUtilities/Source/Private/WorldPartition/HLOD/Builders/HLODBuilderMeshApproximate.cpp
--- a/Engine/Plugins/Editor/WorldPartitionHLODUtilities/Source/Private/WorldPartition/HLOD/Builders/HLODBuilderMeshApproximate.cpp
+++ b/Engine/Plugins/Editor/WorldPartitionHLODUtilities/Source/Private/WorldPartition/HLOD/Builders/HLODBuilderMeshApproximate.cpp
@@ -139,6 +139,7 @@ TArray<UActorComponent*> UHLODBuilderMeshApproximate::Build(const FHLODBuildCont
 				// Rename the asset to its final destination
 				NewAsset->Rename(*AssetName, InHLODBuildContext.AssetsOuter, REN_DontCreateRedirectors | REN_NonTransactional);
 				NewAsset->ClearFlags(RF_Public | RF_Standalone);
+				NewAsset->SetFlags(InHLODBuildContext.AssetsObjectFlags);
 
 				// Clean up flags on the temp package. It is not useful anymore.
 				TempPackage->ClearDirtyFlag();
diff --git a/Engine/Plugins/Editor/WorldPartitionHLODUtilities/Source/Private/WorldPartition/HLOD/Builders/HLODBuilderMeshMerge.cpp b/Engine/Plugins/Editor/WorldPartitionHLODUtilities/Source/Private/WorldPartition/HLOD/Builders/HLODBuilderMeshMerge.cpp
--- a/Engine/Plugins/Editor/WorldPartitionHLODUtilities/Source/Private/WorldPartition/HLOD/Builders/HLODBuilderMeshMerge.cpp
+++ b/Engine/Plugins/Editor/WorldPartitionHLODUtilities/Source/Private/WorldPartition/HLOD/Builders/HLODBuilderMeshMerge.cpp
@@ -102,12 +102,13 @@ TArray<UActorComponent*> UHLODBuilderMeshMerge::Build(const FHLODBuildContext& I
 	MeshMergeUtilities.MergeComponentsToStaticMesh(SourcePrimitiveComponents, InHLODBuildContext.BuildWorld, UseSettings, HLODMaterial, InHLODBuildContext.AssetsOuter->GetPackage(), InHLODBuildContext.AssetsBaseName, Assets, MergedActorLocation, 0.25f, false);
 
 	UStaticMeshComponent* Component = nullptr;
-	Algo::ForEach(Assets, [this, &Component, &MergedActorLocation](UObject* Asset)
+	Algo::ForEach(Assets, [this, &Component, &MergedActorLocation, &InHLODBuildContext](UObject* Asset)
 	{
 		Asset->ClearFlags(RF_Public | RF_Standalone);
 
 		if (Cast<UStaticMesh>(Asset))
 		{
+			Asset->SetFlags(InHLODBuildContext.AssetsObjectFlags);
 			Component = NewObject<UStaticMeshComponent>();
 			Component->SetStaticMesh(static_cast<UStaticMesh*>(Asset));
 			Component->SetWorldLocation(MergedActorLocation);
diff --git a/Engine/Plugins/Editor/WorldPartitionHLODUtilities/Source/Private/WorldPartition/HLOD/Builders/HLODBuilderMeshSimplify.cpp b/Engine/Plugins/Editor/WorldPartitionHLODUtilities/Source/Private/WorldPartition/HLOD/Builders/HLODBuilderMeshSimplify.cpp
--- a/Engine/Plugins/Editor/WorldPartitionHLODUtilities/Source/Private/WorldPartition/HLOD/Builders/HLODBuilderMeshSimplify.cpp
+++ b/Engine/Plugins/Editor/WorldPartitionHLODUtilities/Source/Private/WorldPartition/HLOD/Builders/HLODBuilderMeshSimplify.cpp
@@ -113,12 +113,13 @@ TArray<UActorComponent*> UHLODBuilderMeshSimplify::Build(const FHLODBuildContext
 
 	TArray<UActorComponent*> Components;
 
-	Algo::ForEach(Assets, [this, &Components](UObject* Asset)
+	Algo::ForEach(Assets, [this, &Components, &InHLODBuildContext](UObject* Asset)
 	{
 		Asset->ClearFlags(RF_Public | RF_Standalone);
 
 		if (Cast<UStaticMesh>(Asset))
 		{
+			Asset->SetFlags(InHLODBuildContext.AssetsObjectFlags);
 			UStaticMeshComponent* SMComponent = NewObject<UStaticMeshComponent>();
 			SMComponent->SetStaticMesh(static_cast<UStaticMesh*>(Asset));
 
diff --git a/Engine/Plugins/Editor/WorldPartitionHLODUtilities/Source/Private/WorldPartition/HLOD/Utilities/WorldPartitionHLODUtilities.cpp b/Engine/Plugins/Editor/WorldPartitionHLODUtilities/Source/Private/WorldPartition/HLOD/Utilities/WorldPartitionHLODUtilities.cpp
--- a/Engine/Plugins/Editor/WorldPartitionHLODUtilities/Source/Private/WorldPartition/HLOD/Utilities/WorldPartitionHLODUtilities.cpp
+++ b/Engine/Plugins/Editor/WorldPartitionHLODUtilities/Source/Private/WorldPartition/HLOD/Utilities/WorldPartitionHLODUtilities.cpp
@@ -11,6 +11,7 @@
 #include "WorldPartition/DataLayer/DataLayerInstance.h"
 #include "WorldPartition/HLOD/CustomHLODActor.h"
 #include "WorldPartition/HLOD/HLODActor.h"
+#include "WorldPartition/HLOD/HLODActorExternalResources.h"
 #include "WorldPartition/HLOD/HLODCreationFilter.h"
 #include "WorldPartition/HLOD/HLODLayer.h"
 #include "WorldPartition/HLOD/HLODModifier.h"
@@ -236,6 +237,11 @@ TArray<AWorldPartitionHLOD*> FWorldPartitionHLODUtilities::CreateHLODActors(FHLO
 
 		uint64 HLODActorHash = ComputeHLODActorUniqueHash(HLODLayer, InCreationParams.CellGuid, bCustomHLOD);
 		FName HLODActorName = *FString::Printf(TEXT("%s_%016" UINT64_x_FMT), *HLODLayer->GetName(), HLODActorHash);
+		FName HLODResourcesPackagePath = NAME_None;
+		if (!InCreationParams.HLODResourcesFolder.IsNone())
+		{
+			HLODResourcesPackagePath = (FName(InCreationParams.HLODResourcesFolder.ToString() + HLODActorName.ToString()));
+		}
 
 		AWorldPartitionHLOD* HLODActor = nullptr;
 		FWorldPartitionHandle HLODActorHandle;
@@ -289,6 +295,8 @@ TArray<AWorldPartitionHLOD*> FWorldPartitionHLODUtilities::CreateHLODActors(FHLO
 					HLODActor->AddDataLayer(DataLayerInstance);
 				}
 			}
+
+			HLODActor->SetHLODResourcesPackagePath(HLODResourcesPackagePath);
 		}
 		else
 		{
@@ -417,6 +425,11 @@ TArray<AWorldPartitionHLOD*> FWorldPartitionHLODUtilities::CreateHLODActors(FHLO
 			DirtyReason = TEXT("Standalone");
 		}
 
+		if (HLODActor->GetHLODResourcesPackagePath() != HLODResourcesPackagePath)
+		{
+			UE_LOG(LogHLODBuilder, Error, TEXT("Changing or removing HLOD's external resources package path is not supported. Please delete HLODs and try again."));
+		}
+
 		// If any change was performed, mark HLOD package as dirty
 		if (DirtyReason)
 		{
@@ -560,14 +573,16 @@ void GatherInputStats(AWorldPartitionHLOD* InHLODActor, const TArray<UActorCompo
 
 void GatherOutputStats(AWorldPartitionHLOD* InHLODActor)
 {
-	const UPackage* HLODPackage = InHLODActor->GetPackage();
+	const UPackage* HLODResourcesPackage = InHLODActor->GetHLODResourcesPackage();
+	check(HLODResourcesPackage);
 
 	// Gather relevant assets and process them outside of this ForEach as it's possible that async compilation completion
 	// triggers insertion of new objects, which would break the iteration over UObjects
 	TArray<UStaticMesh*> StaticMeshes;
 	TArray<UTexture*> Textures;
 	TArray<UMaterialInstance*> MaterialInstances;
-	ForEachObjectWithPackage(HLODPackage, [&](UObject* Object)
+
+	ForEachObjectWithPackage(HLODResourcesPackage, [&](UObject* Object)
 	{
 		if (UStaticMesh* StaticMesh = Cast<UStaticMesh>(Object))
 		{
@@ -644,7 +659,7 @@ void GatherOutputStats(AWorldPartitionHLOD* InHLODActor)
 					UTexture* Texture = nullptr;
 					MaterialInstance->GetTextureParameterValue(TextureParamName, Texture, true);
 
-					if (Texture && Texture->GetPackage() == HLODPackage)
+					if (Texture && Texture->GetPackage() == HLODResourcesPackage)
 					{
 						TexturesResourceSize += Texture->GetResourceSizeBytes(EResourceSizeMode::Exclusive);
 						return FMath::RoundToInt64(Texture->GetSurfaceWidth());
@@ -849,26 +864,31 @@ bool FWorldPartitionHLODUtilities::BuildHLOD(const FHLODBuildParams& InBuildPara
 	// Clear stats as we're about to refresh them
 	HLODActor->ResetStats();
 
-	// Rename previous assets found in the HLOD actor package.
+	// Rename previous assets found in the HLOD actor package or HLOD resources package.
 	// Move the previous asset(s) to the transient package, to avoid any object reuse during the build
 	{
 	    TArray<UObject*> ObjectsToRename;
-	    ForEachObjectWithOuter(HLODActor->GetPackage(), [&ObjectsToRename](UObject* Obj)
-	    {
-		    if (!Obj->IsA<AActor>() && !Obj->IsA<UWorld>())
-		    {
-			    ObjectsToRename.Add(Obj);
-		    }
-	    }, false);
+		UPackage* HLODResourcesPackage = HLODActor->GetHLODResourcesPackage();
+
+		if(HLODResourcesPackage)
+		{
+			ForEachObjectWithOuter(HLODResourcesPackage, [&ObjectsToRename](UObject* Obj)
+				{
+					if (!Obj->IsA<AActor>() && !Obj->IsA<UWorld>())
+					{
+						ObjectsToRename.Add(Obj);
+					}
+				}, false);
     
 	    
-	    for (UObject* Obj : ObjectsToRename)
-	    {
-		    // Make sure the old object is not used by anything
-		    Obj->ClearFlags(RF_Standalone | RF_Public);
-		    const FName OldRenamed = MakeUniqueObjectName(GetTransientPackage(), Obj->GetClass(), *FString::Printf(TEXT("OLD_%s"), *Obj->GetName()));
-		    Obj->Rename(*OldRenamed.ToString(), GetTransientPackage(), REN_DontCreateRedirectors | REN_NonTransactional);
-	    }
+			for (UObject* Obj : ObjectsToRename)
+			{
+				// Make sure the old object is not used by anything
+				Obj->ClearFlags(RF_Standalone | RF_Public);
+				const FName OldRenamed = MakeUniqueObjectName(GetTransientPackage(), Obj->GetClass(), *FString::Printf(TEXT("OLD_%s"), *Obj->GetName()));
+				Obj->Rename(*OldRenamed.ToString(), GetTransientPackage(), REN_DontCreateRedirectors | REN_NonTransactional);
+			}
+		}
 	}
 
 	// Gather stats from the input to our HLOD build
@@ -886,14 +906,34 @@ bool FWorldPartitionHLODUtilities::BuildHLOD(const FHLODBuildParams& InBuildPara
 
 			HLODBuilder->SetHLODBuilderSettings(HLODLayer->GetHLODBuilderSettings());
 
+			UPackage* HLODResourcesPackage = nullptr;
+			UHLODActorExternalResources* HLODActorExternalResources = nullptr;
+			if (!HLODActor->GetHLODResourcesPackagePath().IsNone())
+			{
+				HLODResourcesPackage = FindPackage(nullptr, *HLODActor->GetHLODResourcesPackagePath().ToString());
+				if (!HLODResourcesPackage)
+				{
+					HLODResourcesPackage = CreatePackage(*HLODActor->GetHLODResourcesPackagePath().ToString());
+				}
+
+				FString AssetName = FPackageName::GetLongPackageAssetName(HLODActor->GetHLODResourcesPackagePath().ToString());
+				HLODActorExternalResources = FindObject<UHLODActorExternalResources>(HLODResourcesPackage, AssetName);
+				if (!HLODActorExternalResources)
+				{
+					HLODActorExternalResources = NewObject<UHLODActorExternalResources>(HLODResourcesPackage, UHLODActorExternalResources::StaticClass(), FName(AssetName), RF_Public);
+				}
+				HLODActorExternalResources->Resources.Reset();
+			}
+
 			FHLODBuildContext HLODBuildContext;
 			HLODBuildContext.TargetWorld = HLODActor->GetWorld();
 			HLODBuildContext.BuildWorld = BuildHLODWorld;
 			HLODBuildContext.SourceComponents = HLODRelevantComponents;
-			HLODBuildContext.AssetsOuter = HLODActor->GetPackage();
+			HLODBuildContext.AssetsOuter = HLODResourcesPackage ? HLODResourcesPackage : HLODActor->GetPackage();
 			HLODBuildContext.AssetsBaseName = HLODActor->GetActorLabel();
 			HLODBuildContext.MinVisibleDistance = HLODActor->GetMinVisibleDistance();
 			HLODBuildContext.WorldPosition = HLODActor->GetHLODBounds().GetCenter();
+			HLODBuildContext.AssetsObjectFlags = HLODResourcesPackage ? (RF_Public) : RF_NoFlags;
 PRAGMA_DISABLE_DEPRECATION_WARNINGS
 			HLODBuildContext.World = HLODBuildContext.TargetWorld;
 PRAGMA_ENABLE_DEPRECATION_WARNINGS
@@ -977,6 +1017,7 @@ PRAGMA_ENABLE_DEPRECATION_WARNINGS
 				{
 					TArray<UTexture*> Textures;
 					TArray<UStaticMesh*> StaticMeshes;
+					TArray<UMaterialInterface*> Materials;
 				} HLODAssets;				
 
 				// Retrieve assets first, as some of the calls we have to perform can't be done inside of a ForEachObjectWithOuter() (ex: Rename)
@@ -991,11 +1032,21 @@ PRAGMA_ENABLE_DEPRECATION_WARNINGS
 					{
 						HLODAssets.Textures.Add(Texture);
 					}
+
+					if (UMaterialInterface* Material = Cast<UMaterialInterface>(Obj))
+					{
+						HLODAssets.Materials.Add(Material);
+					}
 				});
 
 				// Static meshes
 				for (UStaticMesh* StaticMesh : HLODAssets.StaticMeshes)
 				{
+					if (HLODActorExternalResources)
+					{
+						HLODActorExternalResources->Resources.Add(StaticMesh);
+					}
+
 					// Disable navigation data on HLODs
 					StaticMesh->MarkAsNotHavingNavigationData();
 
@@ -1016,12 +1067,17 @@ PRAGMA_ENABLE_DEPRECATION_WARNINGS
 					}
 				
 					// Rename owned static mesh
-					StaticMesh->Rename(*MakeUniqueObjectName(StaticMesh->GetOuter(), StaticMesh->GetClass(), *FString::Printf(TEXT("StaticMesh_%s"), *HLODLayer->GetName())).ToString());
+					StaticMesh->Rename(*MakeUniqueObjectName(StaticMesh->GetOuter(), StaticMesh->GetClass(), *FString::Printf(TEXT("StaticMesh_%s"), *HLODLayer->GetName())).ToString(), nullptr, REN_DontCreateRedirectors);
 				}
 
 				// Textures
 				for (UTexture* Texture : HLODAssets.Textures)
 				{
+					if (HLODActorExternalResources)
+					{
+						HLODActorExternalResources->Resources.Add(Texture);
+					}
+
 					// Make sure all HLOD textures are assigned to the correct LODGroup
 					if (Texture->LODGroup != TEXTUREGROUP_HierarchicalLOD)
 					{
@@ -1030,6 +1086,15 @@ PRAGMA_ENABLE_DEPRECATION_WARNINGS
 						Texture->PostEditChange();
 					}
 				}
+
+				// Materials
+				for (UMaterialInterface* Material : HLODAssets.Materials)
+				{
+					if (HLODActorExternalResources)
+					{
+						HLODActorExternalResources->Resources.Add(Material);
+					}
+				}
 			}
 
 			HLODActor->SetInputStats(BuildResult.InputStats);
diff --git a/Engine/Plugins/Experimental/Water/Source/Runtime/Private/WaterBodyHLODBuilder.cpp b/Engine/Plugins/Experimental/Water/Source/Runtime/Private/WaterBodyHLODBuilder.cpp
--- a/Engine/Plugins/Experimental/Water/Source/Runtime/Private/WaterBodyHLODBuilder.cpp
+++ b/Engine/Plugins/Experimental/Water/Source/Runtime/Private/WaterBodyHLODBuilder.cpp
@@ -78,6 +78,7 @@ TArray<UActorComponent*> UWaterBodyHLODBuilder::Build(const FHLODBuildContext& I
 	{
 		FName StaticMeshName = MakeUniqueObjectName(nullptr, UStaticMesh::StaticClass(), TEXT("WaterBodyHLODMesh"));
 		UStaticMesh* StaticMesh = NewObject<UStaticMesh>(InHLODBuildContext.AssetsOuter, StaticMeshName);
+		StaticMesh->SetFlags(InHLODBuildContext.AssetsObjectFlags);
 
 		// Mesh
 		{
diff --git a/Engine/Source/Editor/UnrealEd/Private/WorldPartition/WorldPartitionBuilderHelpers.cpp b/Engine/Source/Editor/UnrealEd/Private/WorldPartition/WorldPartitionBuilderHelpers.cpp
--- a/Engine/Source/Editor/UnrealEd/Private/WorldPartition/WorldPartitionBuilderHelpers.cpp
+++ b/Engine/Source/Editor/UnrealEd/Private/WorldPartition/WorldPartitionBuilderHelpers.cpp
@@ -161,7 +161,7 @@ bool FSourceControlHelper::Save(UPackage* Package) const
 	// Save package
 	FString PackageFileName = GetFilename(Package);
 	FSavePackageArgs SaveArgs;
-	SaveArgs.TopLevelFlags = RF_Standalone;
+	SaveArgs.TopLevelFlags = RF_Standalone | RF_Public;
 	SaveArgs.SaveFlags = PackageHelper.UseSourceControl() ? ESaveFlags::SAVE_None : ESaveFlags::SAVE_Async;
 	if (!UPackage::SavePackage(Package, nullptr, *PackageFileName, SaveArgs))
 	{
diff --git a/Engine/Source/Editor/UnrealEd/Private/WorldPartition/WorldPartitionHLODsBuilder.cpp b/Engine/Source/Editor/UnrealEd/Private/WorldPartition/WorldPartitionHLODsBuilder.cpp
--- a/Engine/Source/Editor/UnrealEd/Private/WorldPartition/WorldPartitionHLODsBuilder.cpp
+++ b/Engine/Source/Editor/UnrealEd/Private/WorldPartition/WorldPartitionHLODsBuilder.cpp
@@ -730,17 +730,32 @@ bool UWorldPartitionHLODsBuilder::BuildHLODActors()
 	{
 		if (bSaveActors)
 		{
-			UPackage* ActorPackage = HLODActor->GetPackage();
-			if (ActorPackage->IsDirty())
+			TFunction<bool(UPackage*, const FString&)> SaveHLODPackage = [this, ActorLabel = HLODActor->GetActorLabel()](UPackage* Package, const FString& Description) -> bool
 			{
-				UE_LOG(LogWorldPartitionHLODsBuilder, Display, TEXT("HLOD actor %s was modified, saving..."), *HLODActor->GetActorLabel());
-
-				bool bSaved = SourceControlHelper->Save(ActorPackage);
-				if (!bSaved)
+				if (Package && Package->IsDirty())
 				{
-					UE_LOG(LogWorldPartitionHLODsBuilder, Error, TEXT("Failed to save %s, exiting..."), *USourceControlHelpers::PackageFilename(ActorPackage));
-					return false;
+					UE_LOG(LogWorldPartitionHLODsBuilder, Display, TEXT("%s %s was modified, saving..."), *Description, *ActorLabel);
+
+					bool bSaved = SourceControlHelper->Save(Package);
+					if (!bSaved)
+					{
+						UE_LOG(LogWorldPartitionHLODsBuilder, Error, TEXT("Failed to save %s, exiting..."), *USourceControlHelpers::PackageFilename(Package));
+						return false;
+					}
 				}
+
+				return true;
+			};
+
+			if (!SaveHLODPackage(HLODActor->GetPackage(), TEXT("HLOD Actor")))
+			{
+				return false;
+			}
+
+			// If HLOD resources are stored in a separate package, save it as well
+			if (HLODActor->GetHLODResourcesPackage() != HLODActor->GetPackage())
+			{
+				return SaveHLODPackage(HLODActor->GetHLODResourcesPackage(), TEXT("HLOD resources package"));
 			}
 		}
 
diff --git a/Engine/Source/Runtime/Engine/Private/WorldPartition/HLOD/HLODActor.cpp b/Engine/Source/Runtime/Engine/Private/WorldPartition/HLOD/HLODActor.cpp
--- a/Engine/Source/Runtime/Engine/Private/WorldPartition/HLOD/HLODActor.cpp
+++ b/Engine/Source/Runtime/Engine/Private/WorldPartition/HLOD/HLODActor.cpp
@@ -35,6 +35,7 @@
 #include "WorldPartition/HLOD/HLODActorDesc.h"
 #include "WorldPartition/HLOD/HLODLayer.h"
 #include "WorldPartition/HLOD/HLODStats.h"
+#include "WorldPartition/HLOD/HLODUtils.h"
 #include "WorldPartition/HLOD/IWorldPartitionHLODUtilities.h"
 #include "WorldPartition/HLOD/IWorldPartitionHLODUtilitiesModule.h"
 #include "WorldPartition/HLOD/HLODSourceActorsFromCell.h"
@@ -351,16 +352,7 @@ void AWorldPartitionHLOD::PreSave(FObjectPreSaveContext ObjectSaveContext)
 					// If the HLOD process did create this static mesh
 					if (StaticMesh->GetPackage() == GetPackage())
 					{
-						if (UBodySetup* BodySetup = StaticMesh->GetBodySetup())
-						{
- 							// To ensure a deterministic cook, save the current GUID and restore it below
-							FGuid PreviousBodySetupGuid = BodySetup->BodySetupGuid;
-							BodySetup->DefaultInstance.SetCollisionProfileName(UCollisionProfile::NoCollision_ProfileName);
-							BodySetup->bNeverNeedsCookedCollisionData = true;
-							BodySetup->bHasCookedCollisionData = false;
-							BodySetup->InvalidatePhysicsData();
-							BodySetup->BodySetupGuid = PreviousBodySetupGuid;
-						}
+						UE::Private::HLOD::RemoveCollisionData(StaticMesh);
 					}
 				}
 			}
@@ -440,7 +432,7 @@ void AWorldPartitionHLOD::OnWorldCleanup(UWorld* InWorld, bool bSessionEnded, bo
 	{
 		if (!World->IsGameWorld())
 		{
-			UPackage* HLODPackage = GetPackage();
+			UPackage* HLODResourcesPackage = GetHLODResourcesPackage();
 
 			// Find all assets being edited
 			UAssetEditorSubsystem* AssetEditorSubsystem = GEditor->GetEditorSubsystem<UAssetEditorSubsystem>();
@@ -448,7 +440,7 @@ void AWorldPartitionHLOD::OnWorldCleanup(UWorld* InWorld, bool bSessionEnded, bo
 
 			for (UObject* Asset : AllAssets)
 			{
-				if (Asset->GetPackage() == HLODPackage)
+				if (Asset->GetPackage() == HLODResourcesPackage)
 				{
 					AssetEditorSubsystem->CloseAllEditorsForAsset(Asset);
 				}
@@ -562,6 +554,17 @@ void AWorldPartitionHLOD::SetHLODBounds(const FBox& InBounds)
 	HLODBounds = InBounds;
 }
 
+UPackage* AWorldPartitionHLOD::GetHLODResourcesPackage() const
+{
+	if (!HLODResourcesPackagePath.IsNone())
+	{
+		// If HLODResourcesPackagePath is set, HLOD resources are in that package
+		return FindPackage(nullptr, *HLODResourcesPackagePath.ToString());
+	}
+	// If HLODResourcesPackagePath is not set, HLOD resources are in the actor's package
+	return GetPackage();
+}
+
 void AWorldPartitionHLOD::GetStreamingBounds(FBox& OutRuntimeBounds, FBox& OutEditorBounds) const
 {
 	OutRuntimeBounds = OutEditorBounds = HLODBounds;
diff --git a/Engine/Source/Runtime/Engine/Private/WorldPartition/HLOD/HLODActorExternalResources.cpp b/Engine/Source/Runtime/Engine/Private/WorldPartition/HLOD/HLODActorExternalResources.cpp
new file mode 100644
--- /dev/null
+++ b/Engine/Source/Runtime/Engine/Private/WorldPartition/HLOD/HLODActorExternalResources.cpp
@@ -0,0 +1,26 @@
+// Copyright Epic Games, Inc. All Rights Reserved.
+
+#include "WorldPartition/HLOD/HLODActorExternalResources.h"
+
+#include "Engine/StaticMesh.h"
+#include "WorldPartition/HLOD/HLODUtils.h"
+
+#include UE_INLINE_GENERATED_CPP_BY_NAME(HLODActorExternalResources)
+
+#if WITH_EDITOR
+void UHLODActorExternalResources::PreSave(FObjectPreSaveContext ObjectSaveContext)
+{
+	Super::PreSave(ObjectSaveContext);
+
+	if (ObjectSaveContext.IsCooking())
+	{
+		for (UObject* Resource : Resources)
+		{
+			if (UStaticMesh* StaticMesh = Cast<UStaticMesh>(Resource))
+			{
+				UE::Private::HLOD::RemoveCollisionData(StaticMesh);
+			}
+		}
+	}
+}
+#endif // WITH_EDITOR
\ No newline at end of file
diff --git a/Engine/Source/Runtime/Engine/Private/WorldPartition/HLOD/HLODUtils.cpp b/Engine/Source/Runtime/Engine/Private/WorldPartition/HLOD/HLODUtils.cpp
new file mode 100644
--- /dev/null
+++ b/Engine/Source/Runtime/Engine/Private/WorldPartition/HLOD/HLODUtils.cpp
@@ -0,0 +1,24 @@
+// Copyright Epic Games, Inc. All Rights Reserved.
+
+#include "WorldPartition/HLOD/HLODUtils.h"
+
+#include "Engine/CollisionProfile.h"
+#include "Engine/StaticMesh.h"
+#include "PhysicsEngine/BodySetup.h"
+
+namespace UE::Private::HLOD
+{
+	void RemoveCollisionData(UStaticMesh* InStaticMesh)
+	{
+		if (UBodySetup* BodySetup = InStaticMesh->GetBodySetup())
+		{
+			// To ensure a deterministic cook, save the current GUID and restore it below
+			FGuid PreviousBodySetupGuid = BodySetup->BodySetupGuid;
+			BodySetup->DefaultInstance.SetCollisionProfileName(UCollisionProfile::NoCollision_ProfileName);
+			BodySetup->bNeverNeedsCookedCollisionData = true;
+			BodySetup->bHasCookedCollisionData = false;
+			BodySetup->InvalidatePhysicsData();
+			BodySetup->BodySetupGuid = PreviousBodySetupGuid;
+		}
+	}
+}
\ No newline at end of file
diff --git a/Engine/Source/Runtime/Engine/Private/WorldPartition/HLOD/HLODUtils.h b/Engine/Source/Runtime/Engine/Private/WorldPartition/HLOD/HLODUtils.h
new file mode 100644
--- /dev/null
+++ b/Engine/Source/Runtime/Engine/Private/WorldPartition/HLOD/HLODUtils.h
@@ -0,0 +1,10 @@
+// Copyright Epic Games, Inc. All Rights Reserved.
+
+#pragma once
+
+class UStaticMesh;
+
+namespace UE::Private::HLOD
+{
+	void RemoveCollisionData(UStaticMesh* InStaticMesh);
+}
\ No newline at end of file
diff --git a/Engine/Source/Runtime/Engine/Private/WorldPartition/RuntimeHashSet/WorldPartitionRuntimeHashSetHLODGeneration.cpp b/Engine/Source/Runtime/Engine/Private/WorldPartition/RuntimeHashSet/WorldPartitionRuntimeHashSetHLODGeneration.cpp
--- a/Engine/Source/Runtime/Engine/Private/WorldPartition/RuntimeHashSet/WorldPartitionRuntimeHashSetHLODGeneration.cpp
+++ b/Engine/Source/Runtime/Engine/Private/WorldPartition/RuntimeHashSet/WorldPartitionRuntimeHashSetHLODGeneration.cpp
@@ -429,6 +429,15 @@ bool UWorldPartitionRuntimeHashSet::SetupHLODActors(const IStreamingGenerationCo
 							}
 						}
 					}
+
+					UWorld* SourceWorld = WorldPartition->GetWorld();
+					
+					if (WorldPartition->ShouldExternalizeHLODAssets())
+					{
+						FString FolderPath, PackagePrefix;
+						UWorldPartitionStandaloneHLODSubsystem::GetStandaloneHLODFolderPathAndPackagePrefix(SourceWorld->GetPackage()->GetName(), FolderPath, PackagePrefix);
+						HLODCreationParams.HLODResourcesFolder = FName(FString::Printf(TEXT("%s/Resources/%s%d/"), *FolderPath, *PackagePrefix, HLODLevel));
+					}
 				}
 
 				TArray<AWorldPartitionHLOD*> CellHLODActors = WPHLODUtilities->CreateHLODActors(HLODCreationContext, HLODCreationParams, ActorInstances);
@@ -574,6 +583,21 @@ bool UWorldPartitionRuntimeHashSet::SetupHLODActors(const IStreamingGenerationCo
 			{
 				if (Params.bSaveActors)
 				{
+					if (WorldPartition->ShouldExternalizeHLODAssets())
+					{
+						FWorldPartitionReference HLODActorRef = HLODActorPair.Value.ToReference();
+						if (AWorldPartitionHLOD* HLODActor = CastChecked<AWorldPartitionHLOD>(HLODActorRef.GetActor()))
+						{
+							if (UPackage* HLODResourcesPackage = HLODActor->GetHLODResourcesPackage())
+							{	
+								if (HLODResourcesPackage != HLODActor->GetPackage())
+								{
+									PrivateUtils::DeletePackage(HLODResourcesPackage, Params.SourceControlHelper);
+								}
+							}
+						}
+					}
+
 					PrivateUtils::DeletePackage(WorldPartition, HLODActorPair.Value, Params.SourceControlHelper);
 				}
 				else
diff --git a/Engine/Source/Runtime/Engine/Private/WorldPartition/WorldPartition.cpp b/Engine/Source/Runtime/Engine/Private/WorldPartition/WorldPartition.cpp
--- a/Engine/Source/Runtime/Engine/Private/WorldPartition/WorldPartition.cpp
+++ b/Engine/Source/Runtime/Engine/Private/WorldPartition/WorldPartition.cpp
@@ -388,6 +388,7 @@ UWorldPartition::UWorldPartition(const FObjectInitializer& ObjectInitializer)
 
 	bHasStandaloneHLOD = false;
 	bIsStandaloneHLODWorld = false;
+	bExternalizeHLODAssets = false;
 
 #if WITH_EDITOR
 	bAllowShowingHLODsInEditor = true;
diff --git a/Engine/Source/Runtime/Engine/Public/WorldPartition/HLOD/HLODActor.h b/Engine/Source/Runtime/Engine/Public/WorldPartition/HLOD/HLODActor.h
--- a/Engine/Source/Runtime/Engine/Public/WorldPartition/HLOD/HLODActor.h
+++ b/Engine/Source/Runtime/Engine/Public/WorldPartition/HLOD/HLODActor.h
@@ -116,6 +116,11 @@ public:
 	double GetMinVisibleDistance() const { return MinVisibleDistance; }
 	void SetMinVisibleDistance(double InMinVisibleDistance) { MinVisibleDistance = InMinVisibleDistance; }
 
+	FName GetHLODResourcesPackagePath() const { return HLODResourcesPackagePath; }
+	void SetHLODResourcesPackagePath(FName InHLODResourcesPackagePath) { HLODResourcesPackagePath = InHLODResourcesPackagePath; }
+
+	ENGINE_API UPackage* GetHLODResourcesPackage() const;
+
 	ENGINE_API bool BuildHLOD(bool bForceBuild = false);
 	ENGINE_API bool ShouldBuildHLOD() const;
 
@@ -217,6 +222,9 @@ private:
 
 	UPROPERTY(DuplicateTransient)
 	TArray<TObjectPtr<UHLODRebuildPolicyData>> HLODRebuildPolicyDataSet;
+
+	UPROPERTY()
+	FName HLODResourcesPackagePath;
 #endif
 
 	UPROPERTY()
diff --git a/Engine/Source/Runtime/Engine/Public/WorldPartition/HLOD/HLODActorExternalResources.h b/Engine/Source/Runtime/Engine/Public/WorldPartition/HLOD/HLODActorExternalResources.h
new file mode 100644
--- /dev/null
+++ b/Engine/Source/Runtime/Engine/Public/WorldPartition/HLOD/HLODActorExternalResources.h
@@ -0,0 +1,26 @@
+// Copyright Epic Games, Inc. All Rights Reserved.
+
+#pragma once
+
+#include "Engine/DataAsset.h"
+
+#if WITH_EDITOR
+#include "UObject/ObjectSaveContext.h"
+#endif
+
+#include "HLODActorExternalResources.generated.h"
+
+UCLASS()
+class UHLODActorExternalResources : public UDataAsset
+{
+	GENERATED_BODY()
+
+protected:
+#if WITH_EDITOR
+	ENGINE_API virtual void PreSave(FObjectPreSaveContext ObjectSaveContext) override;
+#endif // WITH_EDITOR
+
+public:
+	UPROPERTY(VisibleAnywhere, Category=Resources)
+	TArray<TObjectPtr<UObject> > Resources;
+};
\ No newline at end of file
diff --git a/Engine/Source/Runtime/Engine/Public/WorldPartition/HLOD/HLODBuilder.h b/Engine/Source/Runtime/Engine/Public/WorldPartition/HLOD/HLODBuilder.h
--- a/Engine/Source/Runtime/Engine/Public/WorldPartition/HLOD/HLODBuilder.h
+++ b/Engine/Source/Runtime/Engine/Public/WorldPartition/HLOD/HLODBuilder.h
@@ -74,6 +74,9 @@ struct FHLODBuildContext
 
 	/** Minimum distance at which the HLOD is expected to be displayed */
 	double MinVisibleDistance;
+
+	/** Flags to set for generated assets */
+	EObjectFlags AssetsObjectFlags;
 };
 
 /**
diff --git a/Engine/Source/Runtime/Engine/Public/WorldPartition/HLOD/IWorldPartitionHLODUtilities.h b/Engine/Source/Runtime/Engine/Public/WorldPartition/HLOD/IWorldPartitionHLODUtilities.h
--- a/Engine/Source/Runtime/Engine/Public/WorldPartition/HLOD/IWorldPartitionHLODUtilities.h
+++ b/Engine/Source/Runtime/Engine/Public/WorldPartition/HLOD/IWorldPartitionHLODUtilities.h
@@ -43,6 +43,7 @@ struct FHLODCreationParams
 	TArray<const UDataLayerInstance*> DataLayerInstances;
 	bool bIsStandalone;
 	TArray<TSharedPtr<IHLODCreationFilter>> Filters;
+	FName HLODResourcesFolder;
 
 	const UExternalDataLayerAsset* GetExternalDataLayerAsset() const
 	{
diff --git a/Engine/Source/Runtime/Engine/Public/WorldPartition/WorldPartition.h b/Engine/Source/Runtime/Engine/Public/WorldPartition/WorldPartition.h
--- a/Engine/Source/Runtime/Engine/Public/WorldPartition/WorldPartition.h
+++ b/Engine/Source/Runtime/Engine/Public/WorldPartition/WorldPartition.h
@@ -506,6 +506,7 @@ public:
 
 	bool HasStandaloneHLOD() const { return bHasStandaloneHLOD; }
 	bool IsStandaloneHLODWorld() const { return bIsStandaloneHLODWorld; }
+	bool ShouldExternalizeHLODAssets() const { return bExternalizeHLODAssets; }
 
 public:
 	static ENGINE_API bool IsSimulating(bool bIncludeTestEnableSimulationStreamingSource = true);
@@ -646,6 +647,9 @@ private:
 	UPROPERTY(EditAnywhere, Category = WorldPartitionSetup, meta = (DisplayName = "Build Standalone HLOD", EditConditionHides, EditCondition = "IsStandaloneHLODAllowed()", HideEditConditionToggle))
 	uint8 bHasStandaloneHLOD : 1;
 
+	UPROPERTY(EditAnywhere, Category = WorldPartitionSetup, meta = (EditConditionHides, EditCondition = "bHasStandaloneHLOD", HideEditConditionToggle))
+	uint8 bExternalizeHLODAssets : 1;
+
 	UPROPERTY()
 	uint8 bIsStandaloneHLODWorld : 1;
 
diff --git a/Engine/Source/Runtime/Landscape/Private/LandscapeHLODBuilder.cpp b/Engine/Source/Runtime/Landscape/Private/LandscapeHLODBuilder.cpp
--- a/Engine/Source/Runtime/Landscape/Private/LandscapeHLODBuilder.cpp
+++ b/Engine/Source/Runtime/Landscape/Private/LandscapeHLODBuilder.cpp
@@ -267,10 +267,11 @@ static UMaterialInterface* BakeLandscapeMaterial(const FHLODBuildContext& InHLOD
 	TArray<UObject*> GeneratedAssets;
 	UMaterialInstanceConstant* LandscapeMaterialInstance = FMaterialUtilities::CreateFlattenMaterialInstance(InHLODBuildContext.AssetsOuter->GetPackage(), MaterialProxySettings, LandscapeMaterial, LandscapeFlattenMaterial, InHLODBuildContext.AssetsBaseName, InLandscapeProxy->GetName(), GeneratedAssets);
 
-	Algo::ForEach(GeneratedAssets, [](UObject* Asset) 
+	Algo::ForEach(GeneratedAssets, [&InHLODBuildContext](UObject* Asset) 
 	{
-		// We don't want any of the generate HLOD assets to be public
+		// We don't want any of the generate HLOD assets to be public, unless required by HLODBuildContext
 		Asset->ClearFlags(RF_Public | RF_Standalone);
+		Asset->SetFlags(InHLODBuildContext.AssetsObjectFlags);
 
 		// Use clamp texture addressing to avoid artifacts between tiles
 		if (UTexture2D* Texture = Cast<UTexture2D>(Asset))
@@ -312,6 +313,7 @@ TArray<UActorComponent*> ULandscapeHLODBuilder::Build(const FHLODBuildContext& I
 	for (ALandscapeProxy* LandscapeProxy : LandscapeProxies)
 	{
 		UStaticMesh* StaticMesh = NewObject<UStaticMesh>(InHLODBuildContext.AssetsOuter);
+		StaticMesh->SetFlags(InHLODBuildContext.AssetsObjectFlags);
 
 		const bool bExportNaniteEnabled = LandscapeProxy->IsNaniteEnabled() || GetDefault<UStaticMesh>()->IsNaniteForceEnabled();
 
