[HLOD] Fix crash when building HLODs when the HLOD Layer cannot be found

#Synced-CL 49363040

diff --git a/Engine/Plugins/Editor/WorldPartitionHLODUtilities/Source/Private/WorldPartition/HLOD/Utilities/WorldPartitionHLODUtilities.cpp b/Engine/Plugins/Editor/WorldPartitionHLODUtilities/Source/Private/WorldPartition/HLOD/Utilities/WorldPartitionHLODUtilities.cpp
--- a/Engine/Plugins/Editor/WorldPartitionHLODUtilities/Source/Private/WorldPartition/HLOD/Utilities/WorldPartitionHLODUtilities.cpp
+++ b/Engine/Plugins/Editor/WorldPartitionHLODUtilities/Source/Private/WorldPartition/HLOD/Utilities/WorldPartitionHLODUtilities.cpp
@@ -895,6 +895,12 @@ bool FWorldPartitionHLODUtilities::BuildHLOD(const FHLODBuildParams& InBuildPara
 	GatherInputStats(HLODActor, HLODRelevantComponents);
 	
 	const UHLODLayer* HLODLayer = HLODActor->GetSourceActors()->GetHLODLayer();
+	if (!HLODLayer)
+	{
+		UE_LOG(LogHLODBuilder, Error, TEXT("FWorldPartitionHLODUtilities::BuildHLOD() - Failed to find HLOD Layer"));
+		return false;
+	}
+
 	TSubclassOf<UHLODBuilder> HLODBuilderClass = GetHLODBuilderClass(HLODLayer);
 
 	if (HLODBuilderClass)
