World Partition - HLOD: Fix possible infinite loop in LoadSourceActors() streaming wait
* StreamAllResources() was called in a while-loop condition causing UpdateResourceStreaming(0, true) to restart in-flight I/O every iteration, and preventing requests from completing within the 0.1s time limit.

#Synced-CL 50781450

diff --git a/Engine/Plugins/Editor/WorldPartitionHLODUtilities/Source/Private/WorldPartition/HLOD/Utilities/WorldPartitionHLODUtilities.cpp b/Engine/Plugins/Editor/WorldPartitionHLODUtilities/Source/Private/WorldPartition/HLOD/Utilities/WorldPartitionHLODUtilities.cpp
--- a/Engine/Plugins/Editor/WorldPartitionHLODUtilities/Source/Private/WorldPartition/HLOD/Utilities/WorldPartitionHLODUtilities.cpp
+++ b/Engine/Plugins/Editor/WorldPartitionHLODUtilities/Source/Private/WorldPartition/HLOD/Utilities/WorldPartitionHLODUtilities.cpp
@@ -758,13 +758,20 @@ bool LoadSourceActors(const AWorldPartitionHLOD* InHLODActor, UWorld* TargetWorl
 		FAssetCompilingManager::Get().ProcessAsyncTasks();
 
 		// Ensure streaming requests are completed
-		// Some requests require the main thread to be ticked in order to run to completion.
+		// The outer loop calls StreamAllResources to discover pending assets and kick off I/O.
+		// The inner loop uses BlockTillAllRequestsFinished to drain in-flight requests without
+		// calling UpdateResourceStreaming, which would restart them and prevent forward progress.
+		// When the inner loop finishes, the outer loop re-checks for assets the tick may have registered.
 		const float STREAMING_WAIT_DT = 0.1f;
 		while (IStreamingManager::Get().StreamAllResources(STREAMING_WAIT_DT) > 0)
 		{
-			// Application tick.
-			FTaskGraphInterface::Get().ProcessThreadUntilIdle(ENamedThreads::GameThread);
-			FTSTicker::GetCoreTicker().Tick(static_cast<float>(FApp::GetDeltaTime()));
+			do
+			{
+				// Application tick.
+				FTaskGraphInterface::Get().ProcessThreadUntilIdle(ENamedThreads::GameThread);
+				FTSTicker::GetCoreTicker().Tick(static_cast<float>(FApp::GetDeltaTime()));
+			}
+			while (IStreamingManager::Get().BlockTillAllRequestsFinished(STREAMING_WAIT_DT, true) > 0);
 		}
 	}
 
