[HLODBuilder] Clean up unnecessary objects from HLOD actor package and external resources package before save

When saving HLOD packages, all objects with flags RF_Standalone or RF_Public are saved, along with their dependencies. Saving objects with RF_Public flag is required to support storing HLOD resources in external packages.
To avoid saving unnecessary objects (e.g. intermediate MeshDescriptions), clear RF_Public flag from all objects that are not among the assets required for this HLOD Actor (StaticMeshes, Textures, Materials) and either HLODActor or HLODActorExternalResources object, depending on whether external HLOD resources are used or not.

#Synced-CL 50763424

diff --git a/Engine/Plugins/Editor/WorldPartitionHLODUtilities/Source/Private/WorldPartition/HLOD/Utilities/WorldPartitionHLODUtilities.cpp b/Engine/Plugins/Editor/WorldPartitionHLODUtilities/Source/Private/WorldPartition/HLOD/Utilities/WorldPartitionHLODUtilities.cpp
--- a/Engine/Plugins/Editor/WorldPartitionHLODUtilities/Source/Private/WorldPartition/HLOD/Utilities/WorldPartitionHLODUtilities.cpp
+++ b/Engine/Plugins/Editor/WorldPartitionHLODUtilities/Source/Private/WorldPartition/HLOD/Utilities/WorldPartitionHLODUtilities.cpp
@@ -982,6 +982,8 @@ PRAGMA_ENABLE_DEPRECATION_WARNINGS
 				}
 			}
 
+			TSet<UObject*> ObjectsToSave;
+
 			// Ideally, this should be performed elsewhere, to allow more flexibility in the HLOD generation
 			for (UActorComponent* HLODComponent : BuildResult.HLODComponents)
 			{
@@ -1060,6 +1062,7 @@ PRAGMA_ENABLE_DEPRECATION_WARNINGS
 					{
 						HLODActorExternalResources->Resources.Add(StaticMesh);
 					}
+					ObjectsToSave.Add(StaticMesh);
 
 					// Disable navigation data on HLODs
 					StaticMesh->MarkAsNotHavingNavigationData();
@@ -1091,6 +1094,7 @@ PRAGMA_ENABLE_DEPRECATION_WARNINGS
 					{
 						HLODActorExternalResources->Resources.Add(Texture);
 					}
+					ObjectsToSave.Add(Texture);
 
 					// Make sure all HLOD textures are assigned to the correct LODGroup
 					if (Texture->LODGroup != TEXTUREGROUP_HierarchicalLOD)
@@ -1108,9 +1112,33 @@ PRAGMA_ENABLE_DEPRECATION_WARNINGS
 					{
 						HLODActorExternalResources->Resources.Add(Material);
 					}
+					ObjectsToSave.Add(Material);
 				}
 			}
 
+			if (HLODActorExternalResources)
+			{
+				ObjectsToSave.Add(HLODActorExternalResources);
+			}
+			else
+			{
+				ObjectsToSave.Add(HLODActor);
+			}
+			
+			// When saving HLOD packages, all objects with flags RF_Standalone or RF_Public are saved, along with their dependencies.
+			// Saving objects with RF_Public flag is required to support storing HLOD resources in external packages.
+			// To avoid saving unnecessary objects (e.g. intermediate MeshDescriptions), clear RF_Public flag from all objects that are not among the assets
+			// required for this HLOD Actor (StaticMeshes, Textures, Materials) and either HLODActor or HLODActorExternalResources object,
+			// depending on whether external HLOD resources are used or not.
+			UPackage* HLODObjectsPackage = HLODResourcesPackage ? HLODResourcesPackage : HLODActor->GetPackage();
+			ForEachObjectWithPackage(HLODObjectsPackage, [&ObjectsToSave](UObject* Obj) {
+				if (!ObjectsToSave.Contains(Obj))
+				{
+					Obj->ClearFlags(RF_Public);
+				}
+				return true;
+			});
+
 			HLODActor->SetInputStats(BuildResult.InputStats);
 			HLODActor->SetHLODComponents(BuildResult.HLODComponents);
 		}
