World Partition - HLOD: Lower the image comparison threshold for triggering HLOD rebuilds
* The default threshold was set to 5%, which was too aggressive - it rejected rebuilds for changes that are still visually noticeable, such as landscape tint or material color adjustments.
* Lowered the default to 1.5%, which avoids skipping noticeable changes while still preventing unnecessary rebuilds when no meaningful visual difference affects the HLOD.

#Synced-CL 51374974

diff --git a/Engine/Plugins/Editor/WorldPartitionHLODUtilities/Source/Public/WorldPartition/HLOD/RebuildPolicies/HLODRebuildPolicyImageCompare.h b/Engine/Plugins/Editor/WorldPartitionHLODUtilities/Source/Public/WorldPartition/HLOD/RebuildPolicies/HLODRebuildPolicyImageCompare.h
--- a/Engine/Plugins/Editor/WorldPartitionHLODUtilities/Source/Public/WorldPartition/HLOD/RebuildPolicies/HLODRebuildPolicyImageCompare.h
+++ b/Engine/Plugins/Editor/WorldPartitionHLODUtilities/Source/Public/WorldPartition/HLOD/RebuildPolicies/HLODRebuildPolicyImageCompare.h
@@ -136,5 +136,5 @@ protected:
 	float MaximumLocalError = 1.00f;
 
 	UPROPERTY(Config)
-	float MaximumGlobalError = 0.05f;
+	float MaximumGlobalError = 0.015f;
 };
