World Partition - HLOD: Make sure all materials are compiled before starting a HLOD build
* The existing calls didn't properly enforce compilation, as FAssetCompilingManager doesn't include shaders
#Synced-CL 51374974

diff --git a/Engine/Plugins/Editor/WorldPartitionHLODUtilities/Source/Private/WorldPartition/HLOD/Utilities/WorldPartitionHLODUtilities.cpp b/Engine/Plugins/Editor/WorldPartitionHLODUtilities/Source/Private/WorldPartition/HLOD/Utilities/WorldPartitionHLODUtilities.cpp
--- a/Engine/Plugins/Editor/WorldPartitionHLODUtilities/Source/Private/WorldPartition/HLOD/Utilities/WorldPartitionHLODUtilities.cpp
+++ b/Engine/Plugins/Editor/WorldPartitionHLODUtilities/Source/Private/WorldPartition/HLOD/Utilities/WorldPartitionHLODUtilities.cpp
@@ -41,12 +41,14 @@
 #include "ProfilingDebugging/ScopedTimers.h"
 #include "Rendering/NaniteResources.h"
 #include "Serialization/ArchiveCrc32.h"
+#include "ShaderCompiler.h"
 #include "StaticMeshCompiler.h"
 #include "StaticMeshResources.h"
 #include "TextureCompiler.h"
 #include "UObject/MetaData.h"
 #include "UObject/GCObjectScopeGuard.h"
 #include "ContentStreaming.h"
+#include "Materials/Material.h"
 
 TAutoConsoleVariable<bool> CVarHLODUseLegacy32BitNameHash(
 	TEXT("wp.Editor.HLOD.UseLegacy32BitNameHash"),
@@ -773,6 +775,29 @@ bool LoadSourceActors(const AWorldPartitionHLOD* InHLODActor, UWorld* TargetWorl
 			}
 			while (IStreamingManager::Get().BlockTillAllRequestsFinished(STREAMING_WAIT_DT, true) > 0);
 		}
+	
+		// Make sure all referenced materials are compiled
+		for (TActorIterator<AActor> It(TargetWorld); It; ++It)
+		{
+			It->ForEachComponent<UPrimitiveComponent>(false, [](UPrimitiveComponent* PrimComp)
+			{
+				TArray<UMaterialInterface*> UsedMaterials;
+				PrimComp->GetUsedMaterials(UsedMaterials);
+				for (UMaterialInterface* Mat : UsedMaterials)
+				{
+					// Walk the full parent chain — the incomplete shader map may be on a parent MIC, not the leaf
+					UMaterialInterface* Current = Mat;
+					while (Current)
+					{
+						Current->EnsureIsComplete();
+						UMaterialInstance* MI = Cast<UMaterialInstance>(Current);
+						Current = MI ? MI->Parent.Get() : nullptr;
+					}
+				}
+			});
+		}
+
+		GShaderCompilingManager->FinishAllCompilation();
 	}
 
 	return bLoadSuccess;
